Fixed FInstancedStruct custom localization gather callback not being consistently registered for commandlets

FStructUtilsEditorModule::StartupModule is not called for commandlets, and so this resulted in the custom localization gather callback being missing for FInstancedStruct.

This caused an issue for the localization commandlet (which fixed it by loading the module directly), but still left other commandlets with the issue, and could mean that things like resaving packages could generate an incomplete localization cache.

This change removes the special fix for the localization commandlet, and instead has the CoreUObject module (which now owns FInstancedStruct) directly register the custom localization gather callback during initialization.

#rnx

diff --git a/Engine/Source/Editor/StructUtilsEditor/Private/StructUtilsEditorModule.cpp b/Engine/Source/Editor/StructUtilsEditor/Private/StructUtilsEditorModule.cpp
--- a/Engine/Source/Editor/StructUtilsEditor/Private/StructUtilsEditorModule.cpp
+++ b/Engine/Source/Editor/StructUtilsEditor/Private/StructUtilsEditorModule.cpp
@@ -6,7 +6,6 @@
 #include "Modules/ModuleManager.h"
 #include "PropertyBagDetails.h"
 #include "PropertyEditorModule.h"
-#include "Serialization/PropertyLocalizationDataGathering.h"
 #include "StructUtils/InstancedStruct.h"
 #include "StructUtils/InstancedStructContainer.h"
 #include "StructUtils/UserDefinedStruct.h"
@@ -20,36 +19,6 @@
 
 IMPLEMENT_MODULE(FStructUtilsEditorModule, StructUtilsEditor)
 
-namespace UE::StructUtils::Private
-{
-
-void GatherForLocalization(const FString& PathToParent, const UScriptStruct* Struct, const void* StructData, const void* DefaultStructData,
-	                       FPropertyLocalizationDataGatherer& PropertyLocalizationDataGatherer, const EPropertyLocalizationGathererTextFlags GatherTextFlags)
-{
-	const FInstancedStruct* ThisInstance = static_cast<const FInstancedStruct*>(StructData);
-	const FInstancedStruct* DefaultInstance = static_cast<const FInstancedStruct*>(DefaultStructData);
-
-	PropertyLocalizationDataGatherer.GatherLocalizationDataFromStruct(PathToParent, Struct, StructData, DefaultStructData, GatherTextFlags);
-
-	if (const UScriptStruct* StructTypePtr = ThisInstance->GetScriptStruct())
-	{
-		const uint8* DefaultInstanceMemory = nullptr;
-		if (DefaultInstance)
-		{
-			// Types must match
-			if (StructTypePtr == DefaultInstance->GetScriptStruct())
-			{
-				DefaultInstanceMemory = DefaultInstance->GetMemory();
-			}
-		}
-
-		PropertyLocalizationDataGatherer.GatherLocalizationDataFromStructWithCallbacks(PathToParent + TEXT(".StructInstance"), StructTypePtr, ThisInstance->GetMemory(),
-			                                                                           DefaultInstanceMemory, GatherTextFlags);
-	}
-}
-
-} // UE::StructUtils::Private
-
 void FStructUtilsEditorModule::StartupModule()
 {
 	// Register the details customizer
@@ -57,13 +26,6 @@ void FStructUtilsEditorModule::StartupModule()
 	PropertyModule.RegisterCustomPropertyTypeLayout("InstancedStruct", FOnGetPropertyTypeCustomizationInstance::CreateStatic(&FInstancedStructDetails::MakeInstance));
 	PropertyModule.RegisterCustomPropertyTypeLayout("InstancedPropertyBag", FOnGetPropertyTypeCustomizationInstance::CreateStatic(&FPropertyBagDetails::MakeInstance));
 	PropertyModule.NotifyCustomizationModuleChanged();
-
-	// Register for localization
-	{
-		static const FAutoRegisterLocalizationDataGatheringCallback AutomaticRegistrationOfLocalizationGatherer(
-			TBaseStructure<FInstancedStruct>::Get(),
-			&UE::StructUtils::Private::GatherForLocalization);
-	}
 }
 
 void FStructUtilsEditorModule::ShutdownModule()
diff --git a/Engine/Source/Editor/UnrealEd/Private/Commandlets/GatherTextFromAssetsCommandlet.cpp b/Engine/Source/Editor/UnrealEd/Private/Commandlets/GatherTextFromAssetsCommandlet.cpp
--- a/Engine/Source/Editor/UnrealEd/Private/Commandlets/GatherTextFromAssetsCommandlet.cpp
+++ b/Engine/Source/Editor/UnrealEd/Private/Commandlets/GatherTextFromAssetsCommandlet.cpp
@@ -906,12 +906,6 @@ bool UGatherTextFromAssetsCommandlet::ParseCommandLineHelper(const FString& InCo
 		}
 	}
 
-	// Preload modules that register custom gather callbacks
-	if (!FModuleManager::LoadModulePtr<IModuleInterface>("StructUtilsEditor"))
-	{
-		UE_LOG(LogGatherTextFromAssetsCommandlet, Warning, TEXT("Failed to preload StructUtilsEditor. FInstancedStruct gathering may be incomplete."));
-	}
-
 	return true;
 }
 
diff --git a/Engine/Source/Runtime/CoreUObject/Private/UObject/CoreNative.cpp b/Engine/Source/Runtime/CoreUObject/Private/UObject/CoreNative.cpp
--- a/Engine/Source/Runtime/CoreUObject/Private/UObject/CoreNative.cpp
+++ b/Engine/Source/Runtime/CoreUObject/Private/UObject/CoreNative.cpp
@@ -11,6 +11,11 @@
 #include "UObject/Stack.h"
 #include "UObject/OverridableManager.h"
 
+#if WITH_EDITORONLY_DATA
+#include "Serialization/PropertyLocalizationDataGathering.h"
+#include "StructUtils/InstancedStruct.h"
+#endif	// WITH_EDITORONLY_DATA
+
 void UClassRegisterAllCompiledInClasses();
 bool IsInAsyncLoadingThreadCoreUObjectInternal();
 bool IsAsyncLoadingCoreUObjectInternal();
@@ -50,6 +55,10 @@ public:
 		IsAsyncLoadingMultithreaded = &IsAsyncLoadingMultithreadedCoreUObjectInternal;
 		GetLoaderType = &GetLoaderTypeInternal;
 
+#if WITH_EDITORONLY_DATA
+		FCoreDelegates::OnPostEngineInit.AddStatic(&RegisterCustomLocalizationDataGathering);
+#endif
+
 		// Register the script callstack callback to the runtime error logging
 #if UE_RAISE_RUNTIME_ERRORS
 		FRuntimeErrors::OnRuntimeIssueLogged.BindStatic(&FCoreUObjectModule::RouteRuntimeMessageToBP);
@@ -62,6 +71,37 @@ public:
 		FFrame::InitPrintScriptCallstack();
 #endif
 	}
+
+#if WITH_EDITORONLY_DATA
+	static void RegisterCustomLocalizationDataGathering()
+	{
+		{
+			static const FAutoRegisterLocalizationDataGatheringCallback _(TBaseStructure<FInstancedStruct>::Get(),
+				[](const FString& PathToParent, const UScriptStruct* Struct, const void* StructData, const void* DefaultStructData, FPropertyLocalizationDataGatherer& PropertyLocalizationDataGatherer, const EPropertyLocalizationGathererTextFlags GatherTextFlags)
+			{
+				const FInstancedStruct* ThisInstance = static_cast<const FInstancedStruct*>(StructData);
+				const FInstancedStruct* DefaultInstance = static_cast<const FInstancedStruct*>(DefaultStructData);
+
+				PropertyLocalizationDataGatherer.GatherLocalizationDataFromStruct(PathToParent, Struct, StructData, DefaultStructData, GatherTextFlags);
+
+				if (const UScriptStruct* StructTypePtr = ThisInstance->GetScriptStruct())
+				{
+					const uint8* DefaultInstanceMemory = nullptr;
+					if (DefaultInstance)
+					{
+						// Types must match
+						if (StructTypePtr == DefaultInstance->GetScriptStruct())
+						{
+							DefaultInstanceMemory = DefaultInstance->GetMemory();
+						}
+					}
+
+					PropertyLocalizationDataGatherer.GatherLocalizationDataFromStructWithCallbacks(PathToParent + TEXT(".StructInstance"), StructTypePtr, ThisInstance->GetMemory(), DefaultInstanceMemory, GatherTextFlags);
+				}
+			});
+		}
+	}
+#endif
 };
 IMPLEMENT_MODULE( FCoreUObjectModule, CoreUObject );
 
