Added AllowedClasses/DisallowedClasses support for the Instanced Struct Viewer

#rnx

diff --git a/Engine/Source/Editor/PropertyEditor/Private/PropertyCustomizationHelpers.cpp b/Engine/Source/Editor/PropertyEditor/Private/PropertyCustomizationHelpers.cpp
--- a/Engine/Source/Editor/PropertyEditor/Private/PropertyCustomizationHelpers.cpp
+++ b/Engine/Source/Editor/PropertyEditor/Private/PropertyCustomizationHelpers.cpp
@@ -1197,6 +1197,42 @@ TArray<const UClass*> PropertyCustomizationHelpers::GetClassesFromMetadataString
 	return Classes;
 }
 
+TArray<const UScriptStruct*> PropertyCustomizationHelpers::GetStructsFromMetadataString(const FString& MetadataString)
+{
+	if (MetadataString.IsEmpty())
+	{
+		return TArray<const UScriptStruct*>();
+	}
+
+	auto FindClass = [](const FString& InStructName) -> const UScriptStruct*
+	{
+		const UScriptStruct* Struct = UClass::TryFindTypeSlow<const UScriptStruct>(InStructName, EFindFirstObjectOptions::EnsureIfAmbiguous);
+		if (!Struct)
+		{
+			Struct = LoadObject<const UScriptStruct>(nullptr, *InStructName);
+		}
+		return Struct;
+	};
+		
+	TArray<FString> StructNames;
+	MetadataString.ParseIntoArrayWS(StructNames, TEXT(","), true);
+
+	TArray<const UScriptStruct*> Structs;
+	Structs.Reserve(StructNames.Num());
+
+	for (const FString& StructName : StructNames)
+	{
+		const UScriptStruct* Struct = FindClass(StructName);
+		if (!Struct)
+		{
+			continue;
+		}
+
+		Structs.Add(Struct);
+	}
+
+	return Structs;
+}
 
 
 TSharedRef<SWidget> PropertyCustomizationHelpers::MakeEditInlineObjectClassPicker(TSharedRef<IPropertyHandle> PropertyHandle, FOnClassPicked OnClassPicked, TSharedPtr<IClassViewerFilter> AdditionalClassFilter)
diff --git a/Engine/Source/Editor/PropertyEditor/Public/PropertyCustomizationHelpers.h b/Engine/Source/Editor/PropertyEditor/Public/PropertyCustomizationHelpers.h
--- a/Engine/Source/Editor/PropertyEditor/Public/PropertyCustomizationHelpers.h
+++ b/Engine/Source/Editor/PropertyEditor/Public/PropertyCustomizationHelpers.h
@@ -227,6 +227,12 @@ namespace PropertyCustomizationHelpers
 	 */
 	PROPERTYEDITOR_API TArray<const UClass*> GetClassesFromMetadataString(const FString& MetadataString);
 
+	/**
+	 * Parse and load the given metadata string into a list of allowed structs.
+	 * The metadata string is likely from something like AllowedClasses or DisallowedClasses.
+	 */
+	PROPERTYEDITOR_API TArray<const UScriptStruct*> GetStructsFromMetadataString(const FString& MetadataString);
+	
 	/**
 	 *
 	 */
diff --git a/Engine/Source/Editor/StructUtilsEditor/Private/SInstancedStructPicker.cpp b/Engine/Source/Editor/StructUtilsEditor/Private/SInstancedStructPicker.cpp
--- a/Engine/Source/Editor/StructUtilsEditor/Private/SInstancedStructPicker.cpp
+++ b/Engine/Source/Editor/StructUtilsEditor/Private/SInstancedStructPicker.cpp
@@ -4,6 +4,7 @@
 #include "DetailLayoutBuilder.h"
 #include "Editor.h"
 #include "IPropertyUtilities.h"
+#include "PropertyCustomizationHelpers.h"
 #include "Modules/ModuleManager.h"
 #include "PropertyHandle.h"
 #include "ScopedTransaction.h"
@@ -58,6 +59,36 @@ namespace UE::StructUtils::Private
 
 bool FInstancedStructFilter::IsStructAllowed(const FStructViewerInitializationOptions& InInitOptions, const UScriptStruct* InStruct, TSharedRef<FStructViewerFilterFuncs> InFilterFuncs)
 {
+	bool bStructAllowed = true;
+	if (!AllowedStructs.IsEmpty())
+	{
+		bStructAllowed = false;
+		for (const TSoftObjectPtr<const UScriptStruct> AllowedStruct : AllowedStructs)
+		{
+			if (InStruct->IsChildOf(AllowedStruct.Get()))
+			{
+				bStructAllowed = true;
+				break;
+			}
+		}
+	}
+	if (!DisallowedStructs.IsEmpty())
+	{
+		for (const TSoftObjectPtr<const UScriptStruct> DisallowedStruct : DisallowedStructs)
+		{
+			if (InStruct->IsChildOf(DisallowedStruct.Get()))
+			{
+				bStructAllowed = false;
+				break;
+			}
+		}
+	}
+
+	if (!bStructAllowed)
+	{
+		return false;
+	}
+	
 	if (InStruct->IsA<UUserDefinedStruct>())
 	{
 		return bAllowUserDefinedStructs;
@@ -228,6 +259,44 @@ TSharedRef<SWidget> SInstancedStructPicker::GenerateStructPicker()
 		}
 
 		StructFilter->AssetReferenceFilter = GEditor->MakeAssetReferenceFilter(AssetReferenceFilterContext);
+
+		auto SoftPointerTransform = [](const UScriptStruct* InStruct) -> TSoftObjectPtr<const UScriptStruct>
+		{
+			return InStruct;
+		};
+		
+		Algo::Transform(PropertyCustomizationHelpers::GetStructsFromMetadataString(StructProperty->GetMetaData("AllowedClasses")), StructFilter->AllowedStructs, SoftPointerTransform);
+		Algo::Transform(PropertyCustomizationHelpers::GetStructsFromMetadataString(StructProperty->GetMetaData("DisallowedClasses")), StructFilter->DisallowedStructs, SoftPointerTransform);
+
+		TArray<UObject*> OwningObjects;
+		StructProperty->GetOuterObjects(OwningObjects);
+		for (UObject* OwningObject : OwningObjects)
+		{
+			if (OwningObject != nullptr)
+			{
+				const FString GetAllowedClassesFunctionName = StructProperty->GetMetaData("GetAllowedClasses");
+				if (!GetAllowedClassesFunctionName.IsEmpty())
+				{
+					const UFunction* GetAllowedClassesFunction = OwningObject ? OwningObject->FindFunction(*GetAllowedClassesFunctionName) : nullptr;
+					if (GetAllowedClassesFunction != nullptr)
+					{
+						DECLARE_DELEGATE_RetVal(TArray<TSoftObjectPtr<UScriptStruct>>, FGetAllowedClasses);
+						StructFilter->AllowedStructs.Append(FGetAllowedClasses::CreateUFunction(OwningObject, GetAllowedClassesFunction->GetFName()).Execute());
+					}
+				}
+		
+				const FString GetDisallowedClassesFunctionName = StructProperty->GetMetaData("GetDisallowedClasses");
+				if (!GetDisallowedClassesFunctionName.IsEmpty())
+				{
+					const UFunction* GetDisallowedClassesFunction = OwningObject ? OwningObject->FindFunction(*GetDisallowedClassesFunctionName) : nullptr;
+					if (GetDisallowedClassesFunction != nullptr)
+					{
+						DECLARE_DELEGATE_RetVal(TArray<TSoftObjectPtr<UScriptStruct>>, FGetDisallowedClasses);
+						StructFilter->DisallowedStructs.Append(FGetDisallowedClasses::CreateUFunction(OwningObject, GetDisallowedClassesFunction->GetFName()).Execute());
+					}
+				}
+			}
+		}
 	}
 
 	const UScriptStruct* SelectedStruct = nullptr;
diff --git a/Engine/Source/Editor/StructUtilsEditor/Public/SInstancedStructPicker.h b/Engine/Source/Editor/StructUtilsEditor/Public/SInstancedStructPicker.h
--- a/Engine/Source/Editor/StructUtilsEditor/Public/SInstancedStructPicker.h
+++ b/Engine/Source/Editor/StructUtilsEditor/Public/SInstancedStructPicker.h
@@ -20,6 +20,12 @@ public:
 	/** The base struct for the property that classes must be a child-of. */
 	TWeakObjectPtr<const UScriptStruct> BaseStruct = nullptr;
 
+	/** The array of allowed structs */
+	TArray<TSoftObjectPtr<const UScriptStruct>> AllowedStructs;
+
+	/** The array of disallowed structs */
+	TArray<TSoftObjectPtr<const UScriptStruct>> DisallowedStructs;
+
 	// A flag controlling whether we allow UserDefinedStructs
 	bool bAllowUserDefinedStructs = false;
 
