Switched to references for loop variables

#rnx

diff --git a/Engine/Source/Editor/StructUtilsEditor/Private/SInstancedStructPicker.cpp b/Engine/Source/Editor/StructUtilsEditor/Private/SInstancedStructPicker.cpp
--- a/Engine/Source/Editor/StructUtilsEditor/Private/SInstancedStructPicker.cpp
+++ b/Engine/Source/Editor/StructUtilsEditor/Private/SInstancedStructPicker.cpp
@@ -63,7 +63,7 @@ bool FInstancedStructFilter::IsStructAllowed(const FStructViewerInitializationOp
 	if (!AllowedStructs.IsEmpty())
 	{
 		bStructAllowed = false;
-		for (const TSoftObjectPtr<const UScriptStruct> AllowedStruct : AllowedStructs)
+		for (const TSoftObjectPtr<const UScriptStruct>& AllowedStruct : AllowedStructs)
 		{
 			if (InStruct->IsChildOf(AllowedStruct.Get()))
 			{
@@ -74,7 +74,7 @@ bool FInstancedStructFilter::IsStructAllowed(const FStructViewerInitializationOp
 	}
 	if (!DisallowedStructs.IsEmpty())
 	{
-		for (const TSoftObjectPtr<const UScriptStruct> DisallowedStruct : DisallowedStructs)
+		for (const TSoftObjectPtr<const UScriptStruct>& DisallowedStruct : DisallowedStructs)
 		{
 			if (InStruct->IsChildOf(DisallowedStruct.Get()))
 			{
