[PropertyBagDetails] propagate the enabled state of the PropertyBag struct to the child rows

diff --git a/Engine/Source/Editor/StructUtilsEditor/Private/Customizations/PropertyBagDetails.cpp b/Engine/Source/Editor/StructUtilsEditor/Private/Customizations/PropertyBagDetails.cpp
--- a/Engine/Source/Editor/StructUtilsEditor/Private/Customizations/PropertyBagDetails.cpp
+++ b/Engine/Source/Editor/StructUtilsEditor/Private/Customizations/PropertyBagDetails.cpp
@@ -804,6 +804,7 @@ void FPropertyBagInstanceDataDetails::OnChildRowAdded(IDetailPropertyRow& ChildR
 	TSharedPtr<IPropertyHandle> ChildPropertyHandle = ChildRow.GetPropertyHandle();
 	
 	bool bSupportedType = true;
+	bool bEditable = true;
 
 	FText UnsupportedTypeWarning;
 
@@ -821,24 +822,31 @@ void FPropertyBagInstanceDataDetails::OnChildRowAdded(IDetailPropertyRow& ChildR
 			UnsupportedTypeWarning = UnsupportedTypeWarningNestedContainer;
 			bSupportedType = false;
 			// Do not allow to edit the data.
-			ChildRow.IsEnabled(false);
+			bEditable = false;
 		}
-		else if ((Desc->ValueType == EPropertyBagPropertyType::UInt32 || Desc->ValueType == EPropertyBagPropertyType::UInt64)
-			&& !bFixedLayout)
+		else
 		{
-			static const FText UnsupportedTypeWarningUnsigned = LOCTEXT("UnsignedTypesWarning", "Unsigned types are not supported throught the property type selection. If you change the type, you will not be able to change it back.");
+			bEditable = BagStructProperty->IsEditable();
 
-			// Warn that the unsinged types cannot be set via the type selection.
-			UnsupportedTypeWarning = UnsupportedTypeWarningUnsigned;
-			bSupportedType = false;
-		}
-		else if (Property->HasMetaData(HideInDetailPanelsName))
-		{
-			ChildRow.Visibility(EVisibility::Hidden);
-			return;
+			if ((Desc->ValueType == EPropertyBagPropertyType::UInt32 || Desc->ValueType == EPropertyBagPropertyType::UInt64)
+				&& !bFixedLayout)
+			{
+				static const FText UnsupportedTypeWarningUnsigned = LOCTEXT("UnsignedTypesWarning", "Unsigned types are not supported throught the property type selection. If you change the type, you will not be able to change it back.");
+
+				// Warn that the unsigned types cannot be set via the type selection.
+				UnsupportedTypeWarning = UnsupportedTypeWarningUnsigned;
+				bSupportedType = false;
+			}
+			else if (Property->HasMetaData(HideInDetailPanelsName))
+			{
+				ChildRow.Visibility(EVisibility::Hidden);
+				return;
+			}
 		}
 	}
 
+	ChildRow.IsEnabled(bEditable);
+
 	if (!bFixedLayout)
 	{
 		// Inline editable name
