Property bag details : expose GetPinFromPropertyDesc utility functio to convert a pin type to a property description

diff --git a/Engine/Source/Editor/StructUtilsEditor/Private/Customizations/PropertyBagDetails.cpp b/Engine/Source/Editor/StructUtilsEditor/Private/Customizations/PropertyBagDetails.cpp
--- a/Engine/Source/Editor/StructUtilsEditor/Private/Customizations/PropertyBagDetails.cpp
+++ b/Engine/Source/Editor/StructUtilsEditor/Private/Customizations/PropertyBagDetails.cpp
@@ -25,7 +25,6 @@
 
 namespace UE::StructUtils
 {
-
 /** Sets property descriptor based on a Blueprint pin type. */
 void SetPropertyDescFromPin(FPropertyBagPropertyDesc& Desc, const FEdGraphPinType& PinType)
 {
@@ -144,6 +143,103 @@
 	}
 }
 
+/** @return Blueprint pin type from property descriptor. */
+FEdGraphPinType GetPropertyDescAsPin(const FPropertyBagPropertyDesc& Desc)
+{
+	UEnum* PropertyTypeEnum = StaticEnum<EPropertyBagPropertyType>();
+	check(PropertyTypeEnum);
+	const UPropertyBagSchema* Schema = GetDefault<UPropertyBagSchema>();
+	check(Schema);
+
+	FEdGraphPinType PinType;
+	PinType.PinSubCategory = NAME_None;
+
+	// Container type
+	//@todo: Handle nested containers in property selection.
+	const EPropertyBagContainerType ContainerType = Desc.ContainerTypes.GetFirstContainerType();
+	switch (ContainerType)
+	{
+	case EPropertyBagContainerType::Array:
+		PinType.ContainerType = EPinContainerType::Array;
+		break;
+	case EPropertyBagContainerType::Set:
+		PinType.ContainerType = EPinContainerType::Set;
+		break;
+	default:
+		PinType.ContainerType = EPinContainerType::None;
+	}
+
+	// Value type
+	switch (Desc.ValueType)
+	{
+	case EPropertyBagPropertyType::Bool:
+		PinType.PinCategory = UEdGraphSchema_K2::PC_Boolean;
+		break;
+	case EPropertyBagPropertyType::Byte:
+		PinType.PinCategory = UEdGraphSchema_K2::PC_Byte;
+		break;
+	case EPropertyBagPropertyType::Int32:
+		PinType.PinCategory = UEdGraphSchema_K2::PC_Int;
+		break;
+	case EPropertyBagPropertyType::Int64:
+		PinType.PinCategory = UEdGraphSchema_K2::PC_Int64;
+		break;
+	case EPropertyBagPropertyType::Float:
+		PinType.PinCategory = UEdGraphSchema_K2::PC_Real;
+		PinType.PinSubCategory = UEdGraphSchema_K2::PC_Float;
+		break;
+	case EPropertyBagPropertyType::Double:
+		PinType.PinCategory = UEdGraphSchema_K2::PC_Real;
+		PinType.PinSubCategory = UEdGraphSchema_K2::PC_Double;
+		break;
+	case EPropertyBagPropertyType::Name:
+		PinType.PinCategory = UEdGraphSchema_K2::PC_Name;
+		break;
+	case EPropertyBagPropertyType::String:
+		PinType.PinCategory = UEdGraphSchema_K2::PC_String;
+		break;
+	case EPropertyBagPropertyType::Text:
+		PinType.PinCategory = UEdGraphSchema_K2::PC_Text;
+		break;
+	case EPropertyBagPropertyType::Enum:
+		// @todo: some pin coloring is not correct due to this (byte-as-enum vs enum). 
+		PinType.PinCategory = UEdGraphSchema_K2::PC_Enum;
+		PinType.PinSubCategoryObject = const_cast<UObject*>(Desc.ValueTypeObject.Get());
+		break;
+	case EPropertyBagPropertyType::Struct:
+		PinType.PinCategory = UEdGraphSchema_K2::PC_Struct;
+		PinType.PinSubCategoryObject = const_cast<UObject*>(Desc.ValueTypeObject.Get());
+		break;
+	case EPropertyBagPropertyType::Object:
+		PinType.PinCategory = UEdGraphSchema_K2::PC_Object;
+		PinType.PinSubCategoryObject = const_cast<UObject*>(Desc.ValueTypeObject.Get());
+		break;
+	case EPropertyBagPropertyType::SoftObject:
+		PinType.PinCategory = UEdGraphSchema_K2::PC_SoftObject;
+		PinType.PinSubCategoryObject = const_cast<UObject*>(Desc.ValueTypeObject.Get());
+		break;
+	case EPropertyBagPropertyType::Class:
+		PinType.PinCategory = UEdGraphSchema_K2::PC_Class;
+		PinType.PinSubCategoryObject = const_cast<UObject*>(Desc.ValueTypeObject.Get());
+		break;
+	case EPropertyBagPropertyType::SoftClass:
+		PinType.PinCategory = UEdGraphSchema_K2::PC_SoftClass;
+		PinType.PinSubCategoryObject = const_cast<UObject*>(Desc.ValueTypeObject.Get());
+		break;
+	case EPropertyBagPropertyType::UInt32:	// Warning : Type only partially supported (Blueprint does not support unsigned type)
+		PinType.PinCategory = UEdGraphSchema_K2::PC_Int;
+		break;
+	case EPropertyBagPropertyType::UInt64:	// Warning : Type only partially supported (Blueprint does not support unsigned type)
+		PinType.PinCategory = UEdGraphSchema_K2::PC_Int64;
+		break;
+	default:
+		ensureMsgf(false, TEXT("Unhandled value type %s"), *UEnum::GetValueAsString(Desc.ValueType));
+		break;
+	}
+
+	return PinType;
+}
+
 namespace Private
 {
 
@@ -329,103 +425,6 @@
 	return !bFound;
 }
 
-/** @return Blueprint pin type from property descriptor. */
-FEdGraphPinType GetPropertyDescAsPin(const FPropertyBagPropertyDesc& Desc)
-{
-	UEnum* PropertyTypeEnum = StaticEnum<EPropertyBagPropertyType>();
-	check(PropertyTypeEnum);
-	const UPropertyBagSchema* Schema = GetDefault<UPropertyBagSchema>();
-	check(Schema);
-
-	FEdGraphPinType PinType;
-	PinType.PinSubCategory = NAME_None;
-
-	// Container type
-	//@todo: Handle nested containers in property selection.
-	const EPropertyBagContainerType ContainerType = Desc.ContainerTypes.GetFirstContainerType();
-	switch (ContainerType)
-	{
-	case EPropertyBagContainerType::Array:
-		PinType.ContainerType = EPinContainerType::Array;
-		break;
-	case EPropertyBagContainerType::Set:
-		PinType.ContainerType = EPinContainerType::Set;
-		break;
-	default:
-		PinType.ContainerType = EPinContainerType::None;
-	}
-
-	// Value type
-	switch (Desc.ValueType)
-	{
-	case EPropertyBagPropertyType::Bool:
-		PinType.PinCategory = UEdGraphSchema_K2::PC_Boolean;
-		break;
-	case EPropertyBagPropertyType::Byte:
-		PinType.PinCategory = UEdGraphSchema_K2::PC_Byte;
-		break;
-	case EPropertyBagPropertyType::Int32:
-		PinType.PinCategory = UEdGraphSchema_K2::PC_Int;
-		break;
-	case EPropertyBagPropertyType::Int64:
-		PinType.PinCategory = UEdGraphSchema_K2::PC_Int64;
-		break;
-	case EPropertyBagPropertyType::Float:
-		PinType.PinCategory = UEdGraphSchema_K2::PC_Real;
-		PinType.PinSubCategory = UEdGraphSchema_K2::PC_Float;
-		break;
-	case EPropertyBagPropertyType::Double:
-		PinType.PinCategory = UEdGraphSchema_K2::PC_Real;
-		PinType.PinSubCategory = UEdGraphSchema_K2::PC_Double;
-		break;
-	case EPropertyBagPropertyType::Name:
-		PinType.PinCategory = UEdGraphSchema_K2::PC_Name;
-		break;
-	case EPropertyBagPropertyType::String:
-		PinType.PinCategory = UEdGraphSchema_K2::PC_String;
-		break;
-	case EPropertyBagPropertyType::Text:
-		PinType.PinCategory = UEdGraphSchema_K2::PC_Text;
-		break;
-	case EPropertyBagPropertyType::Enum:
-		// @todo: some pin coloring is not correct due to this (byte-as-enum vs enum). 
-		PinType.PinCategory = UEdGraphSchema_K2::PC_Enum;
-		PinType.PinSubCategoryObject = const_cast<UObject*>(Desc.ValueTypeObject.Get());
-		break;
-	case EPropertyBagPropertyType::Struct:
-		PinType.PinCategory = UEdGraphSchema_K2::PC_Struct;
-		PinType.PinSubCategoryObject = const_cast<UObject*>(Desc.ValueTypeObject.Get());
-		break;
-	case EPropertyBagPropertyType::Object:
-		PinType.PinCategory = UEdGraphSchema_K2::PC_Object;
-		PinType.PinSubCategoryObject = const_cast<UObject*>(Desc.ValueTypeObject.Get());
-		break;
-	case EPropertyBagPropertyType::SoftObject:
-		PinType.PinCategory = UEdGraphSchema_K2::PC_SoftObject;
-		PinType.PinSubCategoryObject = const_cast<UObject*>(Desc.ValueTypeObject.Get());
-		break;
-	case EPropertyBagPropertyType::Class:
-		PinType.PinCategory = UEdGraphSchema_K2::PC_Class;
-		PinType.PinSubCategoryObject = const_cast<UObject*>(Desc.ValueTypeObject.Get());
-		break;
-	case EPropertyBagPropertyType::SoftClass:
-		PinType.PinCategory = UEdGraphSchema_K2::PC_SoftClass;
-		PinType.PinSubCategoryObject = const_cast<UObject*>(Desc.ValueTypeObject.Get());
-		break;
-	case EPropertyBagPropertyType::UInt32:	// Warning : Type only partially supported (Blueprint does not support unsigned type)
-		PinType.PinCategory = UEdGraphSchema_K2::PC_Int;
-		break;
-	case EPropertyBagPropertyType::UInt64:	// Warning : Type only partially supported (Blueprint does not support unsigned type)
-		PinType.PinCategory = UEdGraphSchema_K2::PC_Int64;
-		break;
-	default:
-		ensureMsgf(false, TEXT("Unhandled value type %s"), *UEnum::GetValueAsString(Desc.ValueType));
-		break;
-	}
-
-	return PinType;
-}
-
 template<typename TFunc>
 void ApplyChangesToPropertyDescs(const FText& SessionName, const TSharedPtr<IPropertyHandle>& StructProperty, const TSharedPtr<IPropertyUtilities>& PropUtils, TFunc&& Function)
 {
@@ -1238,7 +1237,7 @@
 		const FProperty* Property = ChildPropertyHandle->GetProperty();
 		if (FPropertyBagPropertyDesc* Desc = PropertyDescs.FindByPredicate([Property](const FPropertyBagPropertyDesc& Desc){ return Desc.CachedProperty == Property; }))
 		{
-			return UE::StructUtils::Private::GetPropertyDescAsPin(*Desc);
+			return UE::StructUtils::GetPropertyDescAsPin(*Desc);
 		}
 	
 		return FEdGraphPinType();
diff --git a/Engine/Source/Editor/StructUtilsEditor/Public/PropertyBagDetails.h b/Engine/Source/Editor/StructUtilsEditor/Public/PropertyBagDetails.h
--- a/Engine/Source/Editor/StructUtilsEditor/Public/PropertyBagDetails.h
+++ b/Engine/Source/Editor/StructUtilsEditor/Public/PropertyBagDetails.h
@@ -20,6 +20,7 @@ class SWidget;
 namespace UE::StructUtils
 {
 	STRUCTUTILSEDITOR_API void SetPropertyDescFromPin(FPropertyBagPropertyDesc& Desc, const FEdGraphPinType& PinType);
+	STRUCTUTILSEDITOR_API FEdGraphPinType GetPropertyDescAsPin(const FPropertyBagPropertyDesc& Desc);
 }
 
 /**
