Instanced Struct: Handle expansion state properly.
This change fixes the case when structs would always expand when the detail panel is loaded or refreshed, making it a nightmare to work with nested and list of instanced structs.

The new behavior:
- When a detail panel is opened, the previous expansion states are restored from config properly
- When a detail panel is refreshed, all properties keep their current expansion state
- When a new instanced struct is selected from the drop-down, we auto-expand the details below

diff --git a/Engine/Source/Editor/StructUtilsEditor/Private/InstancedStructDetails.cpp b/Engine/Source/Editor/StructUtilsEditor/Private/InstancedStructDetails.cpp
--- a/Engine/Source/Editor/StructUtilsEditor/Private/InstancedStructDetails.cpp
+++ b/Engine/Source/Editor/StructUtilsEditor/Private/InstancedStructDetails.cpp
@@ -281,17 +281,13 @@ void FInstancedStructDataDetails::GetPropertyGroups(const TArray<TSharedPtr<IPro
 
 void FInstancedStructDataDetails::OnStructLayoutChanges()
 {
-	if (StructProvider.IsValid())
-	{
-		// Reset the struct provider immediately, some update functions might get called with the old struct.
-		StructProvider->Reset();
-	}
+	bCanHandleStructValuePostChange = false;
 	OnRegenerateChildren.ExecuteIfBound();
 }
 
 void FInstancedStructDataDetails::OnStructHandlePostChange()
 {
-	if (StructProvider.IsValid())
+	if (bCanHandleStructValuePostChange)
 	{
 		TArray<TWeakObjectPtr<const UStruct>> InstanceTypes = GetInstanceTypes();
 		if (InstanceTypes != CachedInstanceTypes)
@@ -338,20 +334,36 @@ void FInstancedStructDataDetails::GenerateChildContent(IDetailChildrenBuilder& C
 
 		// If the struct has a property customization, then we'll route through AddChildStructure, as it supports
 		// IPropertyTypeCustomization. The other branch is mostly kept as-is for legacy support purposes.
-		IDetailPropertyRow* PropertyRow = ChildBuilder.AddChildStructure(
+		ChildBuilder.AddChildStructure(
 			StructProperty.ToSharedRef(), NewStructProvider, PropertyName, Label
 		);
-
-		// Expansion state is not properly persisted for these structures, so let's expanded it by default for now
-		if (PropertyRow)
-		{
-			PropertyRow->ShouldAutoExpand(true);
-		}
 	}
 	else
 	{
-		TArray<TSharedPtr<IPropertyHandle>> ChildProperties = StructProperty->AddChildStructure(NewStructProvider);
-		
+		uint32 NumChildren = 0;
+		StructProperty->GetNumChildren(NumChildren);
+		TArray<TSharedPtr<IPropertyHandle>> ChildProperties;
+
+		// The property already has the FInstancedStructProvider child property added
+		if (NumChildren > 0)
+		{
+			// So do not add one more, just use the already created property handles
+			if (const TSharedPtr<IPropertyHandle> ChildStructHandle = StructProperty->GetChildHandle(0))
+			{
+				ChildStructHandle->GetNumChildren(NumChildren);
+				ChildProperties.Reserve(NumChildren);
+
+				for (uint32 i = 0; i < NumChildren; ++i)
+				{
+					ChildProperties.Add(ChildStructHandle->GetChildHandle(i));
+				}
+			}
+		}
+		else
+		{
+			ChildProperties = StructProperty->AddChildStructure(NewStructProvider);
+		}
+
 		// Properties may have Category metadata. If that's the case, they should be added under groups.
 		TMap<TSharedPtr<IPropertyHandle>, IDetailGroup*> PropertyToGroup;
 		GetPropertyGroups(ChildProperties, ChildBuilder, PropertyToGroup);
@@ -373,7 +385,7 @@ void FInstancedStructDataDetails::GenerateChildContent(IDetailChildrenBuilder& C
 		}
 	}
 
-	StructProvider = NewStructProvider;
+	bCanHandleStructValuePostChange = true;
 
 	CachedInstanceTypes = GetInstanceTypes();
 }
@@ -418,7 +430,6 @@ void FInstancedStructDetails::CustomizeHeader(TSharedRef<class IPropertyHandle>
 	OnObjectsReinstancedHandle = FCoreUObjectDelegates::OnObjectsReinstanced.AddSP(this, &FInstancedStructDetails::OnObjectsReinstanced);
 
 	HeaderRow
-		.ShouldAutoExpand(true)
 		.NameContent()
 		[
 			StructPropertyHandle->CreatePropertyNameWidget()
diff --git a/Engine/Source/Editor/StructUtilsEditor/Private/SInstancedStructPicker.cpp b/Engine/Source/Editor/StructUtilsEditor/Private/SInstancedStructPicker.cpp
--- a/Engine/Source/Editor/StructUtilsEditor/Private/SInstancedStructPicker.cpp
+++ b/Engine/Source/Editor/StructUtilsEditor/Private/SInstancedStructPicker.cpp
@@ -346,6 +346,9 @@ void SInstancedStructPicker::StructPicked(const UScriptStruct* InStruct)
 		StructProperty->NotifyPostChange(EPropertyChangeType::ValueSet);
 		StructProperty->NotifyFinishedChangingProperties();
 
+		// After the type has changed, let's expand, so that the user can edit the newly appeared child properties
+		StructProperty->SetExpanded(true);
+
 		// Property tree will be invalid after changing the struct type, force update.
 		if (PropUtils.IsValid())
 		{
diff --git a/Engine/Source/Editor/StructUtilsEditor/Public/InstancedStructDetails.h b/Engine/Source/Editor/StructUtilsEditor/Public/InstancedStructDetails.h
--- a/Engine/Source/Editor/StructUtilsEditor/Public/InstancedStructDetails.h
+++ b/Engine/Source/Editor/StructUtilsEditor/Public/InstancedStructDetails.h
@@ -96,14 +96,11 @@ private:
 	/** Handle to the struct property being edited */
 	TSharedPtr<IPropertyHandle> StructProperty;
 
-	/** Struct provider for the structs. */
-	TSharedPtr<FInstancedStructProvider> StructProvider;
-
 	/** Delegate that can be used to refresh the child rows of the current struct (eg, when changing struct type) */
 	FSimpleDelegate OnRegenerateChildren;
 
-	/** True if we're currently handling a StructValuePostChange */
-	bool bIsHandlingStructValuePostChange = false;
+	/** True if we're allowed to handle a StructValuePostChange */
+	bool bCanHandleStructValuePostChange = false;
 	
 	FDelegateHandle UserDefinedStructReinstancedHandle;
 
