Fix crash when reusing Instanced Struct child properties between refreshes. Instead of this, let's just rebuild the child tree.

diff --git a/Engine/Source/Editor/PropertyEditor/Private/PropertyHandleImpl.cpp b/Engine/Source/Editor/PropertyEditor/Private/PropertyHandleImpl.cpp
--- a/Engine/Source/Editor/PropertyEditor/Private/PropertyHandleImpl.cpp
+++ b/Engine/Source/Editor/PropertyEditor/Private/PropertyHandleImpl.cpp
@@ -3585,6 +3585,14 @@ TArray<TSharedPtr<IPropertyHandle>> FPropertyHandleBase::AddChildStructure(TShar
 	return PropertyHandles;
 }
 
+void FPropertyHandleBase::RemoveChildren()
+{
+	if (const TSharedPtr<FPropertyNode> PropertyNode = Implementation->GetPropertyNode())
+	{
+		PropertyNode->RemoveAllChildNodes();
+	}
+}
+
 bool FPropertyHandleBase::CanResetToDefault() const
 {
 	FProperty* Property = GetProperty();
diff --git a/Engine/Source/Editor/PropertyEditor/Private/PropertyHandleImpl.h b/Engine/Source/Editor/PropertyEditor/Private/PropertyHandleImpl.h
--- a/Engine/Source/Editor/PropertyEditor/Private/PropertyHandleImpl.h
+++ b/Engine/Source/Editor/PropertyEditor/Private/PropertyHandleImpl.h
@@ -523,6 +523,7 @@ public:
 	virtual void SetIgnoreValidation(bool bInIgnore) override;
 	virtual TArray<TSharedPtr<IPropertyHandle>> AddChildStructure(TSharedRef<FStructOnScope> ChildStructure) override;
 	virtual TArray<TSharedPtr<IPropertyHandle>> AddChildStructure(TSharedRef<IStructureDataProvider> ChildStructProvider) override;
+	virtual void RemoveChildren() override;
 	virtual bool CanResetToDefault() const override;
 	virtual void ExecuteCustomResetToDefault(const FResetToDefaultOverride& InOnCustomResetToDefault) override;
 	virtual FName GetDefaultCategoryName() const override;
diff --git a/Engine/Source/Editor/PropertyEditor/Private/PropertyNode.cpp b/Engine/Source/Editor/PropertyEditor/Private/PropertyNode.cpp
--- a/Engine/Source/Editor/PropertyEditor/Private/PropertyNode.cpp
+++ b/Engine/Source/Editor/PropertyEditor/Private/PropertyNode.cpp
@@ -494,6 +494,11 @@ void FPropertyNode::AddChildNode(TSharedPtr<FPropertyNode> InNode)
 	ChildNodes.Add(InNode); 
 }
 
+void FPropertyNode::RemoveAllChildNodes()
+{
+	DestroyTree(false);
+}
+
 void FPropertyNode::ClearCachedReadAddresses( bool bRecursive )
 {
 	CachedReadAddresses.Reset();
diff --git a/Engine/Source/Editor/PropertyEditor/Private/PropertyNode.h b/Engine/Source/Editor/PropertyEditor/Private/PropertyNode.h
--- a/Engine/Source/Editor/PropertyEditor/Private/PropertyNode.h
+++ b/Engine/Source/Editor/PropertyEditor/Private/PropertyNode.h
@@ -399,6 +399,11 @@ public:
 	 */
 	void AddChildNode(TSharedPtr<FPropertyNode> InNode);
 
+	/**
+	 * Destroys all child nodes under this node
+	 */
+	void RemoveAllChildNodes();
+
 	/**
 	 * Clears cached read address data
 	 */
diff --git a/Engine/Source/Editor/PropertyEditor/Public/PropertyHandle.h b/Engine/Source/Editor/PropertyEditor/Public/PropertyHandle.h
--- a/Engine/Source/Editor/PropertyEditor/Public/PropertyHandle.h
+++ b/Engine/Source/Editor/PropertyEditor/Public/PropertyHandle.h
@@ -754,6 +754,9 @@ public:
 	virtual TArray<TSharedPtr<IPropertyHandle>> AddChildStructure( TSharedRef<FStructOnScope> ChildStructure ) = 0;
 	virtual TArray<TSharedPtr<IPropertyHandle>> AddChildStructure( TSharedRef<IStructureDataProvider> ChildStructure ) = 0;
 
+	/** Removes all children of this handle */
+	virtual void RemoveChildren() = 0;
+
 	/**
 	 * Returns whether or not the property can be set to default
 	 * 
diff --git a/Engine/Source/Editor/StructUtilsEditor/Private/InstancedStructDetails.cpp b/Engine/Source/Editor/StructUtilsEditor/Private/InstancedStructDetails.cpp
--- a/Engine/Source/Editor/StructUtilsEditor/Private/InstancedStructDetails.cpp
+++ b/Engine/Source/Editor/StructUtilsEditor/Private/InstancedStructDetails.cpp
@@ -340,29 +340,8 @@ void FInstancedStructDataDetails::GenerateChildContent(IDetailChildrenBuilder& C
 	}
 	else
 	{
-		uint32 NumChildren = 0;
-		StructProperty->GetNumChildren(NumChildren);
-		TArray<TSharedPtr<IPropertyHandle>> ChildProperties;
-
-		// The property already has the FInstancedStructProvider child property added
-		if (NumChildren > 0)
-		{
-			// So do not add one more, just use the already created property handles
-			if (const TSharedPtr<IPropertyHandle> ChildStructHandle = StructProperty->GetChildHandle(0))
-			{
-				ChildStructHandle->GetNumChildren(NumChildren);
-				ChildProperties.Reserve(NumChildren);
-
-				for (uint32 i = 0; i < NumChildren; ++i)
-				{
-					ChildProperties.Add(ChildStructHandle->GetChildHandle(i));
-				}
-			}
-		}
-		else
-		{
-			ChildProperties = StructProperty->AddChildStructure(NewStructProvider);
-		}
+		StructProperty->RemoveChildren();
+		TArray<TSharedPtr<IPropertyHandle>> ChildProperties = StructProperty->AddChildStructure(NewStructProvider);
 
 		// Properties may have Category metadata. If that's the case, they should be added under groups.
 		TMap<TSharedPtr<IPropertyHandle>, IDetailGroup*> PropertyToGroup;
