[StructUtils] Drag and drop reordering in StructUtils

+ Add API to PropertyBag to allow re-ordering of properties
+ Add PropertyBagDragDropHandler

#rnx

diff --git a/Engine/Source/Editor/StructUtilsEditor/Private/Customizations/PropertyBagDetails.cpp b/Engine/Source/Editor/StructUtilsEditor/Private/Customizations/PropertyBagDetails.cpp
--- a/Engine/Source/Editor/StructUtilsEditor/Private/Customizations/PropertyBagDetails.cpp
+++ b/Engine/Source/Editor/StructUtilsEditor/Private/Customizations/PropertyBagDetails.cpp
@@ -1,21 +1,24 @@
 // Copyright Epic Games, Inc. All Rights Reserved.
 
 #include "PropertyBagDetails.h"
-#include "DetailWidgetRow.h"
+
 #include "DetailLayoutBuilder.h"
-#include "IPropertyUtilities.h"
+#include "DetailWidgetRow.h"
 #include "IDetailChildrenBuilder.h"
-#include "UObject/EnumProperty.h"
-#include "Widgets/Images/SImage.h"
-#include "Widgets/Input/SButton.h"
-#include "Widgets/Text/SInlineEditableTextBlock.h"
-#include "Widgets/Input/SComboButton.h"
-#include "Framework/MultiBox/MultiBoxBuilder.h"
+#include "IPropertyUtilities.h"
+#include "PropertyBagDragDropHandler.h"
 #include "ScopedTransaction.h"
-#include "StructUtils/UserDefinedStruct.h"
 #include "SPinTypeSelector.h"
 #include "StructUtilsMetadata.h"
+#include "Framework/MultiBox/MultiBoxBuilder.h"
+#include "StructUtils/UserDefinedStruct.h"
 #include "Templates/ValueOrError.h"
+#include "UObject/EnumProperty.h"
+
+#include "Widgets/Images/SImage.h"
+#include "Widgets/Input/SButton.h"
+#include "Widgets/Input/SComboButton.h"
+#include "Widgets/Text/SInlineEditableTextBlock.h"
 
 #include UE_INLINE_GENERATED_CPP_BY_NAME(PropertyBagDetails)
 
@@ -801,19 +804,21 @@ void FPropertyBagInstanceDataDetails::OnChildRowAdded(IDetailPropertyRow& ChildR
 	ChildRow.GetDefaultWidgets(NameWidget, ValueWidget, Row);
 
 	TSharedPtr<IPropertyHandle> ChildPropertyHandle = ChildRow.GetPropertyHandle();
-	
+	check(ChildPropertyHandle);
+
 	bool bSupportedType = true;
 	bool bEditable = true;
 
 	FText UnsupportedTypeWarning;
 
-	TArray<FPropertyBagPropertyDesc> PropertyDescs = UE::StructUtils::Private::GetCommonPropertyDescs(BagStructProperty);
+	const UPropertyBag* BagStruct = BagStructProperty ? UE::StructUtils::Private::GetCommonBagStruct(BagStructProperty) : nullptr;
 	const FProperty* Property = ChildPropertyHandle->GetProperty();
-	if (FPropertyBagPropertyDesc* Desc = PropertyDescs.FindByPredicate([Property](const FPropertyBagPropertyDesc& Desc){ return Desc.CachedProperty == Property; }))
+	const FPropertyBagPropertyDesc* PropertyDesc = BagStruct ? BagStruct->FindPropertyDescByProperty(Property) : nullptr;
+	if (PropertyDesc)
 	{
 		static const FName HideInDetailPanelsName(TEXT("HideInDetailPanel"));
 
-		if (Desc->ContainerTypes.Num() > 1)
+		if (PropertyDesc->ContainerTypes.Num() > 1)
 		{
 			static const FText UnsupportedTypeWarningNestedContainer = LOCTEXT("NestedContainersWarning", "This property type is not supported in the property bag UI.");
 
@@ -827,7 +832,7 @@ void FPropertyBagInstanceDataDetails::OnChildRowAdded(IDetailPropertyRow& ChildR
 		{
 			bEditable = BagStructProperty->IsEditable();
 
-			if ((Desc->ValueType == EPropertyBagPropertyType::UInt32 || Desc->ValueType == EPropertyBagPropertyType::UInt64)
+			if ((PropertyDesc->ValueType == EPropertyBagPropertyType::UInt32 || PropertyDesc->ValueType == EPropertyBagPropertyType::UInt64)
 				&& !bFixedLayout)
 			{
 				static const FText UnsupportedTypeWarningUnsigned = LOCTEXT("UnsignedTypesWarning", "Unsigned types are not supported throught the property type selection. If you change the type, you will not be able to change it back.");
@@ -945,11 +950,39 @@ void FPropertyBagInstanceDataDetails::OnChildRowAdded(IDetailPropertyRow& ChildR
 			ValueWidget.ToSharedRef()
 		];
 
+	TWeakPtr<FPropertyBagInstanceDataDetails> WeakSelf = SharedThis<FPropertyBagInstanceDataDetails>(this);
 
-	if (HasPropertyOverrides())
+	if (PropertyDesc)
 	{
-		TWeakPtr<FPropertyBagInstanceDataDetails> WeakSelf = SharedThis<FPropertyBagInstanceDataDetails>(this);
+		ChildRow.DragDropHandler(MakeShared<FPropertyBagDetailsDragDropHandler>(*PropertyDesc, FOnPropertyBagDetailsRowDropOp::CreateLambda(
+			[WeakSelf, PropertyDesc = *PropertyDesc](const FPropertyBagPropertyDesc& DroppedPropertyDesc, EItemDropZone DropZone) -> FReply
+			{
+				const TSharedPtr<FPropertyBagInstanceDataDetails> DetailsSP = WeakSelf.Pin();
+				const UPropertyBag* ChildBagStruct = DetailsSP ? UE::StructUtils::Private::GetCommonBagStruct(DetailsSP->BagStructProperty) : nullptr;
+				// Validate these properties are still part of the bag.
+				if (!ChildBagStruct || !ChildBagStruct->FindPropertyDescByProperty(PropertyDesc.CachedProperty) || !ChildBagStruct->FindPropertyDescByProperty(DroppedPropertyDesc.CachedProperty))
+				{
+					return FReply::Unhandled();
+				}
+
+				EPropertyBagAlterationResult Result = EPropertyBagAlterationResult::InternalError;
 
+				DetailsSP->BagStructProperty->EnumerateRawData([DropZone, &PropertyDesc, &DroppedPropertyDesc, &Result](void* RawData, const int32 /*DataIndex*/, const int32 /*NumData*/)
+				{
+					if (FInstancedPropertyBag* PropertyBag = RawData ? static_cast<FInstancedPropertyBag*>(RawData) : nullptr)
+					{
+						Result = PropertyBag->ReorderProperty(DroppedPropertyDesc.Name, PropertyDesc.Name, DropZone == EItemDropZone::AboveItem);
+					}
+
+					return true;
+				});
+
+				return Result == EPropertyBagAlterationResult::Success ? FReply::Handled() : FReply::Unhandled();
+			})));
+	}
+
+	if (HasPropertyOverrides())
+	{
 		TAttribute<bool> EditConditionValue = TAttribute<bool>::CreateLambda(
 			[WeakSelf, ChildPropertyHandle]() -> bool
 			{
diff --git a/Engine/Source/Editor/StructUtilsEditor/Private/PropertyBagDragDropHandler.cpp b/Engine/Source/Editor/StructUtilsEditor/Private/PropertyBagDragDropHandler.cpp
new file mode 100644
--- /dev/null
+++ b/Engine/Source/Editor/StructUtilsEditor/Private/PropertyBagDragDropHandler.cpp
@@ -0,0 +1,203 @@
+// Copyright Epic Games, Inc. All Rights Reserved.
+
+#pragma once
+
+#include "PropertyBagDragDropHandler.h"
+
+#include "IPropertyBagEditorGraph.h"
+#include "EdGraph/EdGraph.h"
+#include "Input/Reply.h"
+
+FPropertyBagDetailsDragDropOp::FPropertyBagDetailsDragDropOp(const FPropertyBagPropertyDesc& InPropertyDesc)
+	: PropertyDesc(InPropertyDesc)
+{
+	FGraphEditorDragDropAction::Construct();
+}
+
+void FPropertyBagDetailsDragDropOp::SetDecoration(const EPropertyBagDropState NewDropState)
+{
+	if (CurrentDropState == NewDropState)
+	{
+		return;
+	}
+
+	if (NewDropState == EPropertyBagDropState::Valid)
+	{
+		SetSimpleFeedbackMessage(FAppStyle::Get().GetBrush("Graph.ConnectorFeedback.OK"), FLinearColor::White, NSLOCTEXT("StructUtilsEditor", "ValidTargetTooltip", "Drop property here"));
+		bDropTargetValid = true;
+	}
+	else if (NewDropState == EPropertyBagDropState::SourceIsTarget)
+	{
+		SetSimpleFeedbackMessage(FAppStyle::Get().GetBrush("Graph.ConnectorFeedback.OKWarn"), FLinearColor::White.CopyWithNewOpacity(0.5f), NSLOCTEXT("StructUtilsEditor", "SourceIsTargetTooltip", "Choose a different target"));
+		bDropTargetValid = false;
+	}
+	else // DragDropTarget == EPropertyBagDragDropTargetState::Invalid
+	{
+		check(NewDropState == EPropertyBagDropState::Invalid);
+		SetSimpleFeedbackMessage(FAppStyle::Get().GetBrush("Graph.ConnectorFeedback.Error"), FLinearColor::White, NSLOCTEXT("StructUtilsEditor", "InvalidTargetTooltip", "Invalid target"));
+		bDropTargetValid = false;
+	}
+
+	CurrentDropState = NewDropState;
+	SetDecoratorVisibility(true);
+}
+
+void FPropertyBagDetailsDragDropOp::HoverTargetChanged()
+{
+	IPropertyBagEdGraphDragAndDrop* Graph = Cast<IPropertyBagEdGraphDragAndDrop>(GetHoveredGraph());
+	if (!Graph)
+	{
+		SetDecoration(EPropertyBagDropState::Invalid);
+		SetDecoratorVisibility(false);
+		return;
+	}
+
+	bool bIsValid = false;
+	// Must check in order as the getters will escalate up the chain. Check the pin first.
+	if (const UEdGraphPin* Pin = GetHoveredPin())
+	{
+		bIsValid = Graph->CanReceivePropertyBagDetailsDropOnGraphPin(Pin);
+	} // Check the node next.
+	else if (const UEdGraphNode* Node = GetHoveredNode())
+	{
+		bIsValid = Graph->CanReceivePropertyBagDetailsDropOnGraphNode(Node);
+	}
+	else // Finally, check the graph.
+	{
+		bIsValid = Graph->CanReceivePropertyBagDetailsDropOnGraph(CastChecked<UEdGraph>(Graph));
+	}
+
+	SetDecoration(bIsValid ? EPropertyBagDropState::Valid : EPropertyBagDropState::Invalid);
+}
+
+EVisibility FPropertyBagDetailsDragDropOp::GetIconVisible() const
+{
+	return EVisibility::Visible;
+}
+
+EVisibility FPropertyBagDetailsDragDropOp::GetErrorIconVisible() const
+{
+	// Error icon handled by the decorator directly
+	return EVisibility::Collapsed;
+}
+
+FReply FPropertyBagDetailsDragDropOp::DroppedOnPin(FVector2D ScreenPosition, const FVector2D GraphPosition)
+{
+	if (UEdGraphPin* Pin = GetHoveredPin())
+	{
+		const IPropertyBagEdGraphDragAndDrop* PropertyBagGraph = GetPropertyBagEdGraphDragAndDropInterface();
+		if (PropertyBagGraph && PropertyBagGraph->CanReceivePropertyBagDetailsDropOnGraphPin(Pin))
+		{
+			return PropertyBagGraph->OnPropertyBagDetailsDropOnGraphPin(PropertyDesc, Pin, GraphPosition);
+		}
+	}
+
+	return FReply::Handled();
+}
+
+FReply FPropertyBagDetailsDragDropOp::DroppedOnNode(FVector2D ScreenPosition, const FVector2D GraphPosition)
+{
+	if (UEdGraphNode* Node = GetHoveredNode())
+	{
+		const IPropertyBagEdGraphDragAndDrop* PropertyBagGraph = GetPropertyBagEdGraphDragAndDropInterface();
+		if (PropertyBagGraph && PropertyBagGraph->CanReceivePropertyBagDetailsDropOnGraphNode(Node))
+		{
+			// @todo_pcg: remove return type?
+			return PropertyBagGraph->OnPropertyBagDetailsDropOnGraphNode(PropertyDesc, Node, GraphPosition);
+		}
+	}
+
+	return FReply::Handled();
+}
+
+FReply FPropertyBagDetailsDragDropOp::DroppedOnPanel(const TSharedRef<SWidget>& Panel, const FVector2D ScreenPosition, const FVector2D GraphPosition, UEdGraph& Graph)
+{
+	check(&Graph == GetHoveredGraph())
+	const IPropertyBagEdGraphDragAndDrop* PropertyBagGraph = GetPropertyBagEdGraphDragAndDropInterface();
+	if (PropertyBagGraph && PropertyBagGraph->CanReceivePropertyBagDetailsDropOnGraph(&Graph))
+	{
+		return PropertyBagGraph->OnPropertyBagDetailsDropOnGraph(PropertyDesc, &Graph, GraphPosition);
+	}
+
+	return FReply::Handled();
+}
+
+IPropertyBagEdGraphDragAndDrop* FPropertyBagDetailsDragDropOp::GetPropertyBagEdGraphDragAndDropInterface() const
+{
+	UEdGraph* Graph = GetHoveredGraph();
+	return Graph ? Cast<IPropertyBagEdGraphDragAndDrop>(Graph) : nullptr;
+}
+
+FPropertyBagDetailsDragDropHandler::FPropertyBagDetailsDragDropHandler(const FPropertyBagPropertyDesc& InPropertyDesc, FOnPropertyBagDetailsRowDropOp InOnHandleDragDrop)
+	: PropertyDesc(InPropertyDesc)
+	, OnHandleDetailsRowDragDrop(std::move(InOnHandleDragDrop)) {}
+
+TSharedPtr<FDragDropOperation> FPropertyBagDetailsDragDropHandler::CreateDragDropOperation() const
+{
+	if (TSharedPtr<FPropertyBagDetailsDragDropOp> DragOp = MakeShared<FPropertyBagDetailsDragDropOp>(PropertyDesc))
+	{
+		return DragOp;
+	}
+
+	return nullptr;
+}
+
+TOptional<EItemDropZone> FPropertyBagDetailsDragDropHandler::CanAcceptDrop(const FDragDropEvent& DragDropSource, const EItemDropZone DropZone) const
+{
+	// Property Bag Details Drag and Drop Handler
+	if (const TSharedPtr<FPropertyBagDetailsDragDropOp> DragOp = DragDropSource.GetOperationAs<FPropertyBagDetailsDragDropOp>())
+	{
+		return CanAcceptPropertyBagDetailsDrop(DragOp, DropZone);
+	}
+
+	return TOptional<EItemDropZone>();
+}
+
+bool FPropertyBagDetailsDragDropHandler::AcceptDrop(const FDragDropEvent& DragDropSource, const EItemDropZone DropZone) const
+{
+	// Property Bag Details Drag and Drop Handler
+	if (const TSharedPtr<FPropertyBagDetailsDragDropOp> DragOp = DragDropSource.GetOperationAs<FPropertyBagDetailsDragDropOp>())
+	{
+		return OnPropertyBagDetailsDrop(DragOp, DropZone).IsEventHandled();
+	}
+
+	return false;
+}
+
+TOptional<EItemDropZone> FPropertyBagDetailsDragDropHandler::CanAcceptPropertyBagDetailsDrop(const TSharedPtr<FPropertyBagDetailsDragDropOp>& DragOp, EItemDropZone DropZone) const
+{
+	if (!DragOp.IsValid())
+	{
+		DragOp->SetDecoration(EPropertyBagDropState::Invalid);
+		return TOptional<EItemDropZone>();
+	}
+
+	if (DropZone == EItemDropZone::OntoItem && PropertyDesc.ID != DragOp->PropertyDesc.ID)
+	{
+		DragOp->SetDecoration(EPropertyBagDropState::Invalid);
+		return TOptional<EItemDropZone>();
+	}
+
+	// No effect to drop in these cases. Either source == target, or moving source above/below target puts source in same location.
+	if (PropertyDesc == DragOp->PropertyDesc
+		|| (DropZone == EItemDropZone::AboveItem && DragOp->PropertyDesc.GetCachedIndex() == PropertyDesc.GetCachedIndex() - 1)
+		|| (DropZone == EItemDropZone::BelowItem && DragOp->PropertyDesc.GetCachedIndex() == PropertyDesc.GetCachedIndex() + 1))
+	{
+		DragOp->SetDecoration(EPropertyBagDropState::SourceIsTarget);
+		return TOptional<EItemDropZone>();
+	}
+
+	DragOp->SetDecoration(EPropertyBagDropState::Valid);
+	return DropZone;
+}
+
+FReply FPropertyBagDetailsDragDropHandler::OnPropertyBagDetailsDrop(const TSharedPtr<FPropertyBagDetailsDragDropOp>& DragOp, const EItemDropZone DropZone) const
+{
+	// Currently no functional options for receiving the item onto the property.
+	if (!DragOp.IsValid() || !ensure(DropZone != EItemDropZone::OntoItem) || PropertyDesc == DragOp->PropertyDesc)
+	{
+		return FReply::Unhandled();
+	}
+
+	return OnHandleDetailsRowDragDrop.IsBound() ? OnHandleDetailsRowDragDrop.Execute(DragOp->PropertyDesc, DropZone) : FReply::Unhandled();
+}
diff --git a/Engine/Source/Editor/StructUtilsEditor/Public/IPropertyBagEditorGraph.h b/Engine/Source/Editor/StructUtilsEditor/Public/IPropertyBagEditorGraph.h
new file mode 100644
--- /dev/null
+++ b/Engine/Source/Editor/StructUtilsEditor/Public/IPropertyBagEditorGraph.h
@@ -0,0 +1,105 @@
+// Copyright Epic Games, Inc. All Rights Reserved.
+
+#pragma once
+
+#include "Input/Reply.h"
+#include "UObject/Interface.h"
+
+#include "IPropertyBagEditorGraph.generated.h"
+
+/**
+ * PropertyBag Editor public interface
+ *
+ * Details: IPropertyBagEdGraph is an interface that enables all interactivity between Property Bags and any
+ * Editor Graph.
+ *
+ * Usage: Inherit from `public IPropertyBagEdGraph` along with UEdGraph (or subclass of) in any distinct implementation
+ * of UEdGraph. See implementations below for more details.
+ *
+ * Current Implementations:
+ *   Drag and Drop Operations from a Property Bag Details View Child Row
+ */
+
+struct FPropertyBagPropertyDesc;
+class SWidget;
+class UEdGraphPin;
+class UEdGraphNode;
+class UEdGraph;
+
+UINTERFACE(MinimalAPI)
+class UPropertyBagEdGraphDragAndDrop : public UInterface
+{
+	GENERATED_BODY()
+};
+
+/**
+ * PropertyBag Editor Drag and Drop public interface.
+ *
+ * Usage: This interface can be used independently or as a part of the IPropertyBagEdGraph interface (below), to be
+ * inherited by custom subclasses of UEdGraph. Once inherited, implement the following member overrides to validate and
+ * receive Drag and Drop operations from a PropertyBag details child row. The PropertyBagDragDropHandler
+ * (see PropertyBagDragDropHandler.h/.cpp) inherits from FGraphEditorDragDropAction, mimicking a graph schema action
+ * being dropped into the graph, or onto a pin or node.
+ *
+ * Note: The base handler class currently climbs down ownership from Graph->Pin, so if only the graph is valid, it will
+ * attempt the graph panel drop through the pins' or nodes' parent graph.
+ */
+class STRUCTUTILSEDITOR_API IPropertyBagEdGraphDragAndDrop
+{
+	GENERATED_BODY()
+
+public:
+	virtual bool CanReceivePropertyBagDetailsDropOnGraphPin(const UEdGraphPin* Pin) const;
+	virtual bool CanReceivePropertyBagDetailsDropOnGraphNode(const UEdGraphNode* Node) const;
+	virtual bool CanReceivePropertyBagDetailsDropOnGraph(const UEdGraph* Graph) const;
+	virtual FReply OnPropertyBagDetailsDropOnGraphPin(const FPropertyBagPropertyDesc& PropertyDesc, UEdGraphPin* Pin, FVector2D GraphPosition) const;
+	virtual FReply OnPropertyBagDetailsDropOnGraphNode(const FPropertyBagPropertyDesc& PropertyDesc, UEdGraphNode* Node, FVector2D GraphPosition) const;
+	virtual FReply OnPropertyBagDetailsDropOnGraph(const FPropertyBagPropertyDesc& PropertyDesc, UEdGraph* Graph, FVector2D GraphPosition) const;
+};
+
+UINTERFACE(MinimalAPI)
+class UPropertyBagEdGraph : public UPropertyBagEdGraphDragAndDrop
+{
+	GENERATED_BODY()
+};
+
+/**
+ * PropertyBag Editor public interface
+ *
+ * Usage: This interface to be inherited by custom subclasses of UEdGraph for full implementation of all StructUtils
+ * editor graph integration features.
+ */
+class STRUCTUTILSEDITOR_API IPropertyBagEdGraph : public IPropertyBagEdGraphDragAndDrop
+{
+	GENERATED_BODY()
+};
+
+inline bool IPropertyBagEdGraphDragAndDrop::CanReceivePropertyBagDetailsDropOnGraphPin(const UEdGraphPin* Pin) const
+{
+	return false;
+}
+
+inline bool IPropertyBagEdGraphDragAndDrop::CanReceivePropertyBagDetailsDropOnGraphNode(const UEdGraphNode* Node) const
+{
+	return false;
+}
+
+inline bool IPropertyBagEdGraphDragAndDrop::CanReceivePropertyBagDetailsDropOnGraph(const UEdGraph* Graph) const
+{
+	return false;
+}
+
+inline FReply IPropertyBagEdGraphDragAndDrop::OnPropertyBagDetailsDropOnGraphPin(const FPropertyBagPropertyDesc& PropertyDesc, UEdGraphPin* Pin, const FVector2D GraphPosition) const
+{
+	return FReply::Handled();
+}
+
+inline FReply IPropertyBagEdGraphDragAndDrop::OnPropertyBagDetailsDropOnGraphNode(const FPropertyBagPropertyDesc& PropertyDesc, UEdGraphNode* Node, const FVector2D GraphPosition) const
+{
+	return FReply::Handled();
+}
+
+inline FReply IPropertyBagEdGraphDragAndDrop::OnPropertyBagDetailsDropOnGraph(const FPropertyBagPropertyDesc& PropertyDesc, UEdGraph* Graph, const FVector2D GraphPosition) const
+{
+	return FReply::Handled();
+}
\ No newline at end of file
diff --git a/Engine/Source/Editor/StructUtilsEditor/Public/PropertyBagDetails.h b/Engine/Source/Editor/StructUtilsEditor/Public/PropertyBagDetails.h
--- a/Engine/Source/Editor/StructUtilsEditor/Public/PropertyBagDetails.h
+++ b/Engine/Source/Editor/StructUtilsEditor/Public/PropertyBagDetails.h
@@ -41,9 +41,6 @@ protected:
 	virtual void CustomizeChildren(TSharedRef<IPropertyHandle> StructPropertyHandle, IDetailChildrenBuilder& StructBuilder, IPropertyTypeCustomizationUtils& StructCustomizationUtils) override;
 
 private:
-
-	FReply OnAddProperty() const;
-	
 	/** Handle to the struct property being edited */
 	TSharedPtr<IPropertyHandle> StructProperty;
 	TSharedPtr<IPropertyUtilities> PropUtils;
diff --git a/Engine/Source/Editor/StructUtilsEditor/Public/PropertyBagDragDropHandler.h b/Engine/Source/Editor/StructUtilsEditor/Public/PropertyBagDragDropHandler.h
new file mode 100644
--- /dev/null
+++ b/Engine/Source/Editor/StructUtilsEditor/Public/PropertyBagDragDropHandler.h
@@ -0,0 +1,132 @@
+// Copyright Epic Games, Inc. All Rights Reserved.
+
+#pragma once
+
+#include "GraphEditorDragDropAction.h"
+#include "IDetailDragDropHandler.h"
+#include "StructUtils/PropertyBag.h"
+
+class IPropertyBagEdGraphDragAndDrop;
+
+DECLARE_DELEGATE_RetVal_TwoParams(FReply, FOnPropertyBagDetailsRowDropOp, const FPropertyBagPropertyDesc&, EItemDropZone);
+
+/** State of the drop, useful for source-is-target validation on details rows. */
+enum class EPropertyBagDropState
+{
+	Invalid,
+	Valid,
+	SourceIsTarget
+};
+
+/**
+ * Provides information about the source row (single property) being dragged.
+ * Inherits from FGraphEditorDragDropAction to support dragging to a graph pin/node/panel.
+ */
+class STRUCTUTILSEDITOR_API FPropertyBagDetailsDragDropOp : public FGraphEditorDragDropAction
+{
+public:
+	DRAG_DROP_OPERATOR_TYPE(FPropertyBagDetailsDragDropOp, FGraphEditorDragDropAction);
+
+	/**
+	 * Property Bag Details - Drop Operation Constructor
+	 * @param InPropertyDesc Property description for the dropped property.
+	 */
+	FPropertyBagDetailsDragDropOp(const FPropertyBagPropertyDesc& InPropertyDesc);
+
+	/**
+	 * Sets the pop up widget's icon and text.
+	 * @param NewDropState The state of the drop operation.
+	 */
+	void SetDecoration(EPropertyBagDropState NewDropState);
+
+	// ~Begin FGraphEditorDragDropAction interface
+	/** Event when the drag and drop operation changes hover targets. */
+	virtual void HoverTargetChanged() override;
+
+	/** Visibility of the cursor decorator icon. */
+	virtual EVisibility GetIconVisible() const override;
+	/** Visibility of the cursor decorator error icon. */
+	virtual EVisibility GetErrorIconVisible() const override;
+
+	/**
+	 * The drop operation has been executed on an UEdGraphPin in an Editor Graph.
+	 * @param ScreenPosition 2D screen coordinates.
+	 * @param GraphPosition Local 2D graph coordinates.
+	 * @return If the drop operation was handled by the graph.
+	 */
+	virtual FReply DroppedOnPin(FVector2D ScreenPosition, FVector2D GraphPosition) override;
+
+	/**
+	 * The drop operation has been executed on an UEdGraphNode in an Editor Graph.
+	 * @param ScreenPosition 2D screen coordinates.
+	 * @param GraphPosition Local 2D graph coordinates.
+	 * @return If the drop operation was handled by the graph.
+	 */
+	virtual FReply DroppedOnNode(FVector2D ScreenPosition, FVector2D GraphPosition) override;
+
+	/**
+	 * The drop operation has been executed on an SGraphPanel in an Editor Graph.
+	 * @param Panel The graph panel widget to receive the drop action.
+	 * @param ScreenPosition 2D screen coordinates.
+	 * @param GraphPosition Local 2D graph coordinates.
+	 * @param Graph The graph owner of the target panel.
+	 * @return If the drop operation was handled by the graph.
+	 */
+	virtual FReply DroppedOnPanel(const TSharedRef<SWidget>& Panel, FVector2D ScreenPosition, FVector2D GraphPosition, UEdGraph& Graph) override;
+	// ~End FGraphEditorDragDropAction interface
+
+	/** The property description of the dragged property. */
+	const FPropertyBagPropertyDesc PropertyDesc;
+	/** Cached state for the decorator. */
+	EPropertyBagDropState CurrentDropState = EPropertyBagDropState::Invalid;
+
+private:
+	IPropertyBagEdGraphDragAndDrop* GetPropertyBagEdGraphDragAndDropInterface() const;
+};
+
+/** Handles drag-and-drop (as both source and target) for a single property's child widget row */
+class STRUCTUTILSEDITOR_API FPropertyBagDetailsDragDropHandler : public IDetailDragDropHandler
+{
+public:
+	/**
+	 * Property Bag Details - Drag/Drop Handler Constructor
+	 * @param InPropertyDesc Property description for the child row.
+	 * @param InOnHandleDragDrop Custom function for handling receipt of a drop operation for this Property Bag details view.
+	 */
+	FPropertyBagDetailsDragDropHandler(const FPropertyBagPropertyDesc& InPropertyDesc, FOnPropertyBagDetailsRowDropOp InOnHandleDragDrop);
+
+	/**
+	 * Creates the drag and drop operation object that contains necessary information for a drag and drop operation.
+	 * @return Shared pointer to the drag and drop operation.
+	 */
+	virtual TSharedPtr<FDragDropOperation> CreateDragDropOperation() const override;
+
+	/**
+	 * Whether the details view child row can accept a drag and drop operation. 
+	 * @param DragDropSource The source dragged onto this row.
+	 * @param DropZone Zone of the drop, relative to this child row.
+	 * @return [Optional] Target zone, which can be overridden. False if the row will not accept.
+	 */
+	virtual TOptional<EItemDropZone> CanAcceptDrop(const FDragDropEvent& DragDropSource, EItemDropZone DropZone) const override;
+
+	/**
+	 * Accept and handle the drag and drop operation.
+	 * @param DragDropSource The source dragged onto this row.
+	 * @param DropZone Zone of the drop, relative to this child row.
+	 * @return True if the operation was successful.
+	 */
+	virtual bool AcceptDrop(const FDragDropEvent& DragDropSource, EItemDropZone DropZone) const override;
+
+private:
+	/** CanAcceptDrop helper for Property Bag Details row cases, specifically. */
+	TOptional<EItemDropZone> CanAcceptPropertyBagDetailsDrop(const TSharedPtr<FPropertyBagDetailsDragDropOp>& DragOp, EItemDropZone DropZone) const;
+
+	/** AcceptDrop helper for Property Bag Details row cases, specifically. */
+	FReply OnPropertyBagDetailsDrop(const TSharedPtr<FPropertyBagDetailsDragDropOp>& DragOp, EItemDropZone DropZone) const;
+
+	/** The current row's property bag property description. */
+	const FPropertyBagPropertyDesc PropertyDesc;
+
+	/** Callback for widget to determine if the drag to a details view child row was handled successfully. Signature: FReply(const FPropertyBagPropertyDesc&, EItemDropZone)*/
+	FOnPropertyBagDetailsRowDropOp OnHandleDetailsRowDragDrop;
+};
diff --git a/Engine/Source/Runtime/CoreUObject/Private/StructUtils/PropertyBag.cpp b/Engine/Source/Runtime/CoreUObject/Private/StructUtils/PropertyBag.cpp
--- a/Engine/Source/Runtime/CoreUObject/Private/StructUtils/PropertyBag.cpp
+++ b/Engine/Source/Runtime/CoreUObject/Private/StructUtils/PropertyBag.cpp
@@ -19,6 +19,13 @@ static TAutoConsoleVariable<bool> CVarEnforceValidName(
 	true,
 	TEXT("Enforce valid naming on new and renamed property bag properties."));
 
+#if WITH_EDITOR
+static TAutoConsoleVariable<bool> CVarValidatePropertyBagDescriptionEquality(
+	TEXT("StructUtils.ValidatePropertyBagDescriptionEquality"),
+	false,
+	TEXT("Do a full check on description equality. Otherwise, rely on the unique Guid."));
+#endif // WITH_EDITOR
+
 #include UE_INLINE_GENERATED_CPP_BY_NAME(PropertyBag)
 
 struct FPropertyBagCustomVersion
@@ -1516,6 +1523,56 @@ namespace UE::StructUtils::Private
 			return EPropertyBagAlterationResult::SourcePropertyNotFound;
 		}
 	}
+
+	EPropertyBagAlterationResult ReorderProperty(TArray<FPropertyBagPropertyDesc>& Descs, const int32 SourcePropertyIndex, const int32 TargetPropertyIndex, const bool bInsertBefore)
+	{
+		// If the source and the target are the same, the reordering needs no operation.
+		if (SourcePropertyIndex == TargetPropertyIndex)
+		{
+			return EPropertyBagAlterationResult::NoOperation;
+		}
+		else if (!Descs.IsValidIndex(SourcePropertyIndex))
+		{
+			return EPropertyBagAlterationResult::SourcePropertyNotFound;
+		}
+		else if (!Descs.IsValidIndex(TargetPropertyIndex))
+		{
+			return EPropertyBagAlterationResult::TargetPropertyNotFound;
+		}
+
+		TArray<FPropertyBagPropertyDesc> NewDescs;
+		NewDescs.Reserve(Descs.Num());
+
+		const int32 Min = std::min(SourcePropertyIndex, TargetPropertyIndex);
+		const int32 SourceSpacer = SourcePropertyIndex < TargetPropertyIndex;
+
+		// Elements before the source and target properties.
+		int32 ReadIndex;
+		for (ReadIndex = 0; ReadIndex < Min; ++ReadIndex)
+		{
+			NewDescs.Emplace(std::move(Descs[ReadIndex]));
+		}
+
+		// Between the source and target properties (skip source if needed)
+		for (ReadIndex += SourceSpacer; ReadIndex < TargetPropertyIndex; ++ReadIndex)
+		{
+			NewDescs.Emplace(std::move(Descs[ReadIndex]));
+		}
+
+		// Source and target property
+		NewDescs.Emplace(std::move(Descs[bInsertBefore ? SourcePropertyIndex : TargetPropertyIndex]));
+		NewDescs.Emplace(std::move(Descs[bInsertBefore ? TargetPropertyIndex : SourcePropertyIndex]));
+
+		// The remaining elements
+		for (++ReadIndex; ReadIndex < Descs.Num(); ++ReadIndex)
+		{
+			NewDescs.Emplace(std::move(Descs[ReadIndex]));
+		}
+
+		Descs = std::move(NewDescs);
+
+		return EPropertyBagAlterationResult::Success;
+	}
 } // UE::StructUtils::Private
 
 
@@ -1723,6 +1780,33 @@ bool FPropertyBagPropertyDesc::IsClassType() const
 	}
 }
 
+bool FPropertyBagPropertyDesc::operator==(const FPropertyBagPropertyDesc& OtherDesc) const
+{
+	if (ID.IsValid() && ID == OtherDesc.ID)
+	{
+#if WITH_EDITOR
+		if (CVarValidatePropertyBagDescriptionEquality.GetValueOnAnyThread())
+		{
+			check(CachedProperty == OtherDesc.CachedProperty);
+			check(ValueType == OtherDesc.ValueType);
+			check(ValueTypeObject == OtherDesc.ValueTypeObject);
+			check(Name == OtherDesc.Name);
+			check(ContainerTypes == OtherDesc.ContainerTypes);
+			check(PropertyFlags == OtherDesc.PropertyFlags);
+			check(CachedIndex == OtherDesc.CachedIndex);
+			check(MetaData == OtherDesc.MetaData);
+			check(MetaClass == OtherDesc.MetaClass);
+		}
+#endif // WITH_EDITOR
+
+		return true;
+	}
+	else
+	{
+		return false;
+	}
+}
+
 bool FPropertyBagPropertyDesc::CompatibleType(const FPropertyBagPropertyDesc& Other) const
 {
 	// Containers must match
@@ -1910,7 +1994,39 @@ EPropertyBagAlterationResult FInstancedPropertyBag::RenameProperty(const FName P
 	}
 
 	const EPropertyBagAlterationResult Result = UE::StructUtils::Private::RenameProperty(Descs, PropertyToRename, NewName);
+	if (Result == EPropertyBagAlterationResult::Success)
+	{
+		const UPropertyBag* NewBagStruct = UPropertyBag::GetOrCreateFromDescs(Descs);
+		MigrateToNewBagStruct(NewBagStruct);
+	}
+
+	return Result;
+}
+
+EPropertyBagAlterationResult FInstancedPropertyBag::ReorderProperty(const FName SourcePropertyName, const FName TargetPropertyName, const bool bInsertBefore)
+{
+	int32 SourceIndex = INDEX_NONE;
+	int32 TargetIndex = INDEX_NONE;
+	if (const UPropertyBag* CurrentBagStruct = GetPropertyBagStruct())
+	{
+		const FPropertyBagPropertyDesc* SourceDesc = CurrentBagStruct->FindPropertyDescByName(SourcePropertyName);
+		SourceIndex = SourceDesc ? SourceDesc->GetCachedIndex() : INDEX_NONE;
+		const FPropertyBagPropertyDesc* TargetDesc = CurrentBagStruct->FindPropertyDescByName(TargetPropertyName);
+		TargetIndex = TargetDesc ? TargetDesc->GetCachedIndex() : INDEX_NONE;
+	}
+
+	return ReorderProperty(SourceIndex, TargetIndex, bInsertBefore);
+}
+
+EPropertyBagAlterationResult FInstancedPropertyBag::ReorderProperty(const int32 SourcePropertyIndex, const int32 TargetPropertyIndex, const bool bInsertBefore)
+{
+	TArray<FPropertyBagPropertyDesc> Descs;
+	if (const UPropertyBag* CurrentBagStruct = GetPropertyBagStruct())
+	{
+		Descs = CurrentBagStruct->GetPropertyDescs();
+	}
 
+	const EPropertyBagAlterationResult Result = UE::StructUtils::Private::ReorderProperty(Descs, SourcePropertyIndex, TargetPropertyIndex, bInsertBefore);
 	if (Result == EPropertyBagAlterationResult::Success)
 	{
 		const UPropertyBag* NewBagStruct = UPropertyBag::GetOrCreateFromDescs(Descs);
@@ -3770,7 +3886,8 @@ const UPropertyBag* UPropertyBag::GetOrCreateFromDescs(const TConstArrayView<FPr
 #endif
 			NewProperty->SetPropertyFlags((EPropertyFlags)Desc.PropertyFlags);
 			NewBag->AddCppProperty(NewProperty);
-			Desc.CachedProperty = NewProperty; 
+			Desc.CachedProperty = NewProperty;
+			Desc.CachedIndex = DescIndex;
 		}
 	}
 
@@ -3900,7 +4017,12 @@ const FPropertyBagPropertyDesc* UPropertyBag::FindPropertyDescByPropertyName(con
 
 const FPropertyBagPropertyDesc* UPropertyBag::FindPropertyDescByProperty(const FProperty* Property) const
 {
-	return PropertyDescs.FindByPredicate([&Property](const FPropertyBagPropertyDesc& Desc) { return Desc.CachedProperty && Desc.CachedProperty == Property; });
+	return Property ? PropertyDescs.FindByPredicate([&Property](const FPropertyBagPropertyDesc& Desc) { return Desc.CachedProperty && Desc.CachedProperty == Property; }) : nullptr;
+}
+
+const FPropertyBagPropertyDesc* UPropertyBag::FindPropertyDescByIndex(const int32 Index) const
+{
+	return PropertyDescs.IsValidIndex(Index) ? &PropertyDescs[Index] : nullptr;
 }
 
 bool UPropertyBag::OwnsPropertyDesc(const FPropertyBagPropertyDesc& Desc) const
diff --git a/Engine/Source/Runtime/CoreUObject/Public/StructUtils/PropertyBag.h b/Engine/Source/Runtime/CoreUObject/Public/StructUtils/PropertyBag.h
--- a/Engine/Source/Runtime/CoreUObject/Public/StructUtils/PropertyBag.h
+++ b/Engine/Source/Runtime/CoreUObject/Public/StructUtils/PropertyBag.h
@@ -172,12 +172,14 @@ enum class EPropertyBagResult : uint8
 UENUM()
 enum class EPropertyBagAlterationResult : uint8
 {
-	Success,                       // Operation succeeded.
-	InternalError,                 // The operation could not be completed, due to an internal property bag error.
-	PropertyNameEmpty,             // The property name is empty.
-	PropertyNameInvalidCharacters, // The property name contains illegal characters. Consider using FInstancedPropertyBag::SanitizePropertyName.
-	SourcePropertyNotFound,        // The source property or property name was not found.
-	TargetPropertyAlreadyExists    // The target property or property name already exists.
+	Success,                             // Operation succeeded.
+	NoOperation = Success UMETA(Hidden), // No operation was necessary to warrant a successful operation. Semantic alias.
+	InternalError,                       // The operation could not be completed, due to an internal property bag error.
+	PropertyNameEmpty,                   // The property name is empty.
+	PropertyNameInvalidCharacters,       // The property name contains illegal characters. Consider using FInstancedPropertyBag::SanitizePropertyName.
+	SourcePropertyNotFound,              // The source property or property name was not found.
+	TargetPropertyNotFound,              // The target property or property name was not found.
+	TargetPropertyAlreadyExists          // The target property or property name already exists.
 };
 
 USTRUCT()
@@ -200,6 +202,11 @@ struct COREUOBJECT_API FPropertyBagPropertyDescMetaData
 
 	void Serialize(FArchive& Ar);
 
+	bool operator==(const FPropertyBagPropertyDescMetaData& Other) const
+	{
+		return Key == Other.Key && Value == Other.Value;
+	}
+
 	FORCEINLINE friend FArchive& operator<<(FArchive& Ar, FPropertyBagPropertyDescMetaData& PropertyDescMetaData)
 	{
 		PropertyDescMetaData.Serialize(Ar);
@@ -279,6 +286,12 @@ struct COREUOBJECT_API FPropertyBagPropertyDesc
 	/** @return true if the property type is class or soft class */
 	bool IsClassType() const;
 
+	/** @return index of the property after constructed in the bag */
+	int32 GetCachedIndex() const { return CachedIndex; }
+
+	/** @return true if the property bag's unique IDs match */
+	bool operator==(const FPropertyBagPropertyDesc& OtherDesc) const;
+
 	/** Pointer to object that defines the Enum, Struct, or Class. */
 	UPROPERTY(EditAnywhere, Category="Default")
 	TObjectPtr<const UObject> ValueTypeObject = nullptr;
@@ -315,6 +328,12 @@ struct COREUOBJECT_API FPropertyBagPropertyDesc
 
 	/** Cached property pointer, set in UPropertyBag::GetOrCreateFromDescs. */
 	const FProperty* CachedProperty = nullptr;
+
+private:
+	friend class UPropertyBag;
+
+	/** Index of the property in the bag, set in UPropertyBag::GetOrCreateFromDescs. */
+	int32 CachedIndex = INDEX_NONE;
 };
 
 /**
@@ -482,6 +501,24 @@ struct COREUOBJECT_API FInstancedPropertyBag
 	 */
 	EPropertyBagAlterationResult RenameProperty(FName PropertyToRename, FName NewName);
 
+	/**
+	 * Reorders a property either before or after another target property.
+	 * @param SourcePropertyName The source property to insert elsewhere in the array.
+	 * @param TargetPropertyName The target property to conduct the reorder around (before or after).
+	 * @param bInsertBefore Insert the source property before the target property. If false, it will be inserted after.
+	 * @return The result of the alteration.
+	 */
+	EPropertyBagAlterationResult ReorderProperty(FName SourcePropertyName, FName TargetPropertyName, bool bInsertBefore = true);
+
+	/**
+	 * Reorders a property either before or after another target property.
+	 * @param SourcePropertyIndex The index of the source property to insert elsewhere in the array.
+	 * @param TargetPropertyIndex The index of the target property to conduct the reorder around (before or after).
+	 * @param bInsertBefore Insert the source property before the target property. If false, it will be inserted after.
+	 * @return The result of the alteration.
+	 */
+	EPropertyBagAlterationResult ReorderProperty(int32 SourcePropertyIndex, int32 TargetPropertyIndex, bool bInsertBefore = true);
+
 	/**
 	 * Changes the type of this bag and migrates existing values.
 	 * The properties are matched between the bags based on the property ID.
@@ -1248,6 +1285,9 @@ public:
 	/** @return property description based on pointer to property. */
 	const FPropertyBagPropertyDesc* FindPropertyDescByProperty(const FProperty* Property) const;
 
+	/** @return property description based on index. */
+	const FPropertyBagPropertyDesc* FindPropertyDescByIndex(int32 Index) const;
+
 	/** @return true if we own the supplied property description, false otherwise. */
 	bool OwnsPropertyDesc(const FPropertyBagPropertyDesc& Desc) const;
 
