CIS fix

#rnx

diff --git a/Engine/Source/Editor/StructUtilsEditor/Private/PropertyBagDragDropHandler.cpp b/Engine/Source/Editor/StructUtilsEditor/Private/PropertyBagDragDropHandler.cpp
--- a/Engine/Source/Editor/StructUtilsEditor/Private/PropertyBagDragDropHandler.cpp
+++ b/Engine/Source/Editor/StructUtilsEditor/Private/PropertyBagDragDropHandler.cpp
@@ -81,21 +81,21 @@ EVisibility FPropertyBagDetailsDragDropOp::GetErrorIconVisible() const
 	return EVisibility::Collapsed;
 }
 
-FReply FPropertyBagDetailsDragDropOp::DroppedOnPin(FVector2D ScreenPosition, const FVector2D GraphPosition)
+FReply FPropertyBagDetailsDragDropOp::DroppedOnPin(const FVector2f& ScreenPosition, const FVector2f& GraphPosition)
 {
 	if (UEdGraphPin* Pin = GetHoveredPin())
 	{
 		const IPropertyBagEdGraphDragAndDrop* PropertyBagGraph = GetPropertyBagEdGraphDragAndDropInterface();
 		if (PropertyBagGraph && PropertyBagGraph->CanReceivePropertyBagDetailsDropOnGraphPin(Pin))
 		{
-			return PropertyBagGraph->OnPropertyBagDetailsDropOnGraphPin(PropertyDesc, Pin, GraphPosition);
+			return PropertyBagGraph->OnPropertyBagDetailsDropOnGraphPin(PropertyDesc, Pin, FDeprecateSlateVector2D(GraphPosition));
 		}
 	}
 
 	return FReply::Handled();
 }
 
-FReply FPropertyBagDetailsDragDropOp::DroppedOnNode(FVector2D ScreenPosition, const FVector2D GraphPosition)
+FReply FPropertyBagDetailsDragDropOp::DroppedOnNode(const FVector2f& ScreenPosition, const FVector2f& GraphPosition)
 {
 	if (UEdGraphNode* Node = GetHoveredNode())
 	{
@@ -103,20 +103,20 @@ FReply FPropertyBagDetailsDragDropOp::DroppedOnNode(FVector2D ScreenPosition, co
 		if (PropertyBagGraph && PropertyBagGraph->CanReceivePropertyBagDetailsDropOnGraphNode(Node))
 		{
 			// @todo_pcg: remove return type?
-			return PropertyBagGraph->OnPropertyBagDetailsDropOnGraphNode(PropertyDesc, Node, GraphPosition);
+			return PropertyBagGraph->OnPropertyBagDetailsDropOnGraphNode(PropertyDesc, Node, FDeprecateSlateVector2D(GraphPosition));
 		}
 	}
 
 	return FReply::Handled();
 }
 
-FReply FPropertyBagDetailsDragDropOp::DroppedOnPanel(const TSharedRef<SWidget>& Panel, const FVector2D ScreenPosition, const FVector2D GraphPosition, UEdGraph& Graph)
+FReply FPropertyBagDetailsDragDropOp::DroppedOnPanel(const TSharedRef<SWidget>& Panel, const FVector2f& ScreenPosition, const FVector2f& GraphPosition, UEdGraph& Graph)
 {
 	check(&Graph == GetHoveredGraph())
 	const IPropertyBagEdGraphDragAndDrop* PropertyBagGraph = GetPropertyBagEdGraphDragAndDropInterface();
 	if (PropertyBagGraph && PropertyBagGraph->CanReceivePropertyBagDetailsDropOnGraph(&Graph))
 	{
-		return PropertyBagGraph->OnPropertyBagDetailsDropOnGraph(PropertyDesc, &Graph, GraphPosition);
+		return PropertyBagGraph->OnPropertyBagDetailsDropOnGraph(PropertyDesc, &Graph, FDeprecateSlateVector2D(GraphPosition));
 	}
 
 	return FReply::Handled();
diff --git a/Engine/Source/Editor/StructUtilsEditor/Public/PropertyBagDragDropHandler.h b/Engine/Source/Editor/StructUtilsEditor/Public/PropertyBagDragDropHandler.h
--- a/Engine/Source/Editor/StructUtilsEditor/Public/PropertyBagDragDropHandler.h
+++ b/Engine/Source/Editor/StructUtilsEditor/Public/PropertyBagDragDropHandler.h
@@ -54,7 +54,7 @@ public:
 	 * @param GraphPosition Local 2D graph coordinates.
 	 * @return If the drop operation was handled by the graph.
 	 */
-	virtual FReply DroppedOnPin(FVector2D ScreenPosition, FVector2D GraphPosition) override;
+	virtual FReply DroppedOnPin(const FVector2f& ScreenPosition, const FVector2f& GraphPosition) override;
 
 	/**
 	 * The drop operation has been executed on an UEdGraphNode in an Editor Graph.
@@ -62,7 +62,7 @@ public:
 	 * @param GraphPosition Local 2D graph coordinates.
 	 * @return If the drop operation was handled by the graph.
 	 */
-	virtual FReply DroppedOnNode(FVector2D ScreenPosition, FVector2D GraphPosition) override;
+	virtual FReply DroppedOnNode(const FVector2f& ScreenPosition, const FVector2f& GraphPosition) override;
 
 	/**
 	 * The drop operation has been executed on an SGraphPanel in an Editor Graph.
@@ -72,7 +72,7 @@ public:
 	 * @param Graph The graph owner of the target panel.
 	 * @return If the drop operation was handled by the graph.
 	 */
-	virtual FReply DroppedOnPanel(const TSharedRef<SWidget>& Panel, FVector2D ScreenPosition, FVector2D GraphPosition, UEdGraph& Graph) override;
+	virtual FReply DroppedOnPanel(const TSharedRef<SWidget>& Panel, const FVector2f& ScreenPosition, const FVector2f& GraphPosition, UEdGraph& Graph) override;
 	// ~End FGraphEditorDragDropAction interface
 
 	/** The property description of the dragged property. */
