[StructUtils] Added new features to the Property Bag Details View Child Row Widget

+ Customizable via the EPropertyBagChildRowFeatures enum (set via metadata specifier)
+ New STypeSelector widget (wrapper around SPinTypeSelector) for type selection
+ New Access Specifier button
+ New drop down menu

diff --git a/Engine/Source/Editor/StructUtilsEditor/Private/Customizations/PropertyBagDetails.cpp b/Engine/Source/Editor/StructUtilsEditor/Private/Customizations/PropertyBagDetails.cpp
--- a/Engine/Source/Editor/StructUtilsEditor/Private/Customizations/PropertyBagDetails.cpp
+++ b/Engine/Source/Editor/StructUtilsEditor/Private/Customizations/PropertyBagDetails.cpp
@@ -8,8 +8,8 @@
 #include "IPropertyUtilities.h"
 #include "PropertyBagDragDropHandler.h"
 #include "ScopedTransaction.h"
-#include "SPinTypeSelector.h"
 #include "StructUtilsMetadata.h"
+#include "STypeSelector.h" // Custom widget for pill type selector
 #include "Framework/MultiBox/MultiBoxBuilder.h"
 #include "StructUtils/UserDefinedStruct.h"
 #include "Templates/ValueOrError.h"
@@ -353,29 +353,6 @@ void SetPropertyDescs(const TSharedPtr<IPropertyHandle>& StructProperty, const T
 	}
 }
 
-/** @return sanitized property name based on the input string. */
-FName GetValidPropertyName(const FString& Name)
-{
-	FName Result;
-	if (!Name.IsEmpty())
-	{
-		if (!FName::IsValidXName(Name, INVALID_OBJECTNAME_CHARACTERS))
-		{
-			Result = MakeObjectNameFromDisplayLabel(Name, NAME_None);
-		}
-		else
-		{
-			Result = FName(Name);
-		}
-	}
-	else
-	{
-		Result = FName(TEXT("Property"));
-	}
-
-	return Result;
-}
-
 FName GetPropertyNameSafe(const TSharedPtr<IPropertyHandle>& PropertyHandle)
 {
 	const FProperty* Property = PropertyHandle ? PropertyHandle->GetProperty() : nullptr;
@@ -514,6 +491,19 @@ bool FindUserFunction(TSharedPtr<IPropertyHandle> InStructProperty, FName InFunc
 	return OutTarget != nullptr && OutFunc != nullptr;
 }
 
+FText GetAccessSpecifierNameFromFlags(const EPropertyFlags Flags)
+{
+	// TODO: Support 'protected'. For now treat protected and private the same.
+	if (!!(Flags & (CPF_NativeAccessSpecifierPrivate | CPF_NativeAccessSpecifierProtected)))
+	{
+		return LOCTEXT("AccessSpecifierPrivate", "Private");
+	}
+	else // Public flag or not, should be treated as public.
+	{
+		return LOCTEXT("AccessSpecifierPublic", "Public");
+	}
+}
+
 // UFunction calling helpers. 
 // Use our "own" hardcoded reflection system for types used in UFunctions calls in this file.
 template<typename T> struct TypeName { static const TCHAR* Get() = delete; };
@@ -775,8 +765,207 @@ void CopyPropertyValue(const FPropertyBagPropertyDesc* InSourcePropertyDesc, con
 	InSourcePropertyDesc->CachedProperty->CopyCompleteValue(TargetValueAddress, SourceValueAddress);
 }
 
+void GetFilteredVariableTypeTree(const TSharedPtr<IPropertyHandle>& BagStructProperty, TArray<TSharedPtr<UEdGraphSchema_K2::FPinTypeTreeInfo>>& TypeTree, const ETypeTreeFilter TypeTreeFilter)
+{
+	// The type selector popup might outlive this details view, so bag struct property can be invalid here.
+	if (!BagStructProperty || !BagStructProperty->IsValidHandle())
+	{
+		return;
+	}
+
+	UFunction* IsPinTypeAcceptedFunc = nullptr;
+	UObject* IsPinTypeAcceptedTarget = nullptr;
+	if (FindUserFunction(BagStructProperty, Metadata::IsPinTypeAcceptedName, IsPinTypeAcceptedFunc, IsPinTypeAcceptedTarget))
+	{
+		check(IsPinTypeAcceptedFunc && IsPinTypeAcceptedTarget);
+
+		// We need to make sure the signature matches perfectly: bool(FEdGraphPinType, bool)
+		bool bFuncIsValid = ValidateFunctionSignature<bool, FEdGraphPinType, bool>(IsPinTypeAcceptedFunc);
+		if (!ensureMsgf(bFuncIsValid, TEXT("[%s] Function %s does not have the right signature."), *Metadata::IsPinTypeAcceptedName.ToString(), *IsPinTypeAcceptedFunc->GetName()))
+		{
+			return;
+		}
+	}
+
+	auto IsPinTypeAccepted = [IsPinTypeAcceptedFunc, IsPinTypeAcceptedTarget](const FEdGraphPinType& InPinType, bool bInIsChild) -> bool
+	{
+		if (IsPinTypeAcceptedFunc && IsPinTypeAcceptedTarget)
+		{
+			const TValueOrError<bool, void> bIsValid = CallFunc<bool>(IsPinTypeAcceptedTarget, IsPinTypeAcceptedFunc, InPinType, bInIsChild);
+			return bIsValid.HasValue() && bIsValid.GetValue();
+		}
+		else
+		{
+			return true;
+		}
+	};
+
+	check(GetDefault<UEdGraphSchema_K2>());
+	TArray<TSharedPtr<UEdGraphSchema_K2::FPinTypeTreeInfo>> TempTypeTree;
+	GetDefault<UPropertyBagSchema>()->GetVariableTypeTree(TempTypeTree, TypeTreeFilter);
+
+	// Filter
+	for (TSharedPtr<UEdGraphSchema_K2::FPinTypeTreeInfo>& PinType : TempTypeTree)
+	{
+		if (!PinType.IsValid() || !IsPinTypeAccepted(PinType->GetPinType(/*bForceLoadSubCategoryObject*/false), /*bInIsChild=*/ false))
+		{
+			continue;
+		}
+
+		for (int32 ChildIndex = 0; ChildIndex < PinType->Children.Num();)
+		{
+			TSharedPtr<UEdGraphSchema_K2::FPinTypeTreeInfo> Child = PinType->Children[ChildIndex];
+			if (Child.IsValid())
+			{
+				const FEdGraphPinType& ChildPinType = Child->GetPinType(/*bForceLoadSubCategoryObject*/false);
+
+				if (!CanHaveMemberVariableOfType(ChildPinType) || !IsPinTypeAccepted(ChildPinType, /*bInIsChild=*/ true))
+				{
+					PinType->Children.RemoveAt(ChildIndex);
+					continue;
+				}
+			}
+			++ChildIndex;
+		}
+
+		TypeTree.Add(PinType);
+	}
+}
+
+bool CanDeleteProperty(const TSharedPtr<IPropertyHandle>& InStructProperty, const TSharedPtr<IPropertyHandle>& ChildPropertyHandle)
+{
+	if (!InStructProperty || !InStructProperty->IsValidHandle() || !ChildPropertyHandle || !ChildPropertyHandle->IsValidHandle())
+	{
+		return false;
+	}
+
+	// Extra check provided by the user to cancel a remove action. Useful to provide the user a possibility to cancel the action if
+	// the given property is in use elsewhere.
+	UFunction* CanRemovePropertyFunc = nullptr;
+	UObject* CanRemovePropertyTarget = nullptr;
+	if (FindUserFunction(InStructProperty, Metadata::CanRemovePropertyName, CanRemovePropertyFunc, CanRemovePropertyTarget))
+	{
+		check(CanRemovePropertyFunc && CanRemovePropertyTarget);
+
+		FName PropertyName = ChildPropertyHandle->GetProperty()->GetFName();
+		const UPropertyBag* PropertyBag = GetCommonBagStruct(InStructProperty);
+		const FPropertyBagPropertyDesc* PropertyDesc = PropertyBag ? PropertyBag->FindPropertyDescByName(PropertyName) : nullptr;
+
+		if (!PropertyDesc)
+		{
+			return false;
+		}
+
+		// We need to make sure the signature matches perfectly: bool(FGuid, FName)s
+		bool bFuncIsValid = UE::StructUtils::Private::ValidateFunctionSignature<bool, FGuid, FName>(CanRemovePropertyFunc);
+		if (!ensureMsgf(bFuncIsValid, TEXT("[%s] Function %s does not have the right signature."), *UE::StructUtils::Metadata::CanRemovePropertyName.ToString(), *CanRemovePropertyFunc->GetName()))
+		{
+			return false;
+		}
+
+		const TValueOrError<bool, void> bCanRemove = UE::StructUtils::Private::CallFunc<bool>(CanRemovePropertyTarget, CanRemovePropertyFunc, PropertyDesc->ID, PropertyDesc->Name);
+
+		if (bCanRemove.HasError() || !bCanRemove.GetValue())
+		{
+			return false;
+		}
+	}
+
+	return true;
+}
+
+void DeleteProperty(TSharedPtr<IPropertyHandle> InStructProperty, TSharedPtr<IPropertyHandle> ChildPropertyHandle, const TSharedPtr<IPropertyUtilities>& PropUtils)
+{
+	if (!InStructProperty || !InStructProperty->IsValidHandle() || !ChildPropertyHandle || !ChildPropertyHandle->IsValidHandle())
+	{
+		return;
+	}
+
+	if (!CanDeleteProperty(InStructProperty, ChildPropertyHandle))
+	{
+		return;
+	}
+
+	ApplyChangesToPropertyDescs(
+		FText::Format(LOCTEXT("OnPropertyDeleted", "Deleted property: {0}"), ChildPropertyHandle->GetPropertyDisplayName()),
+		InStructProperty,
+		PropUtils,
+		[&ChildPropertyHandle](TArray<FPropertyBagPropertyDesc>& PropertyDescs)
+		{
+			const FProperty* Property = ChildPropertyHandle ? ChildPropertyHandle->GetProperty() : nullptr;
+			PropertyDescs.RemoveAll([Property](const FPropertyBagPropertyDesc& Desc) { return Desc.CachedProperty == Property; });
+		});
+}
+
+FEdGraphPinType GetPinInfo(const TSharedPtr<IPropertyHandle>& ChildPropertyHandle, const TSharedPtr<IPropertyHandle>& InBagStructProperty)
+{
+	// The SPinTypeSelector popup might outlive this details view, so bag struct property can be invalid here.
+	if (!InBagStructProperty || !InBagStructProperty->IsValidHandle() || !ChildPropertyHandle || !ChildPropertyHandle->IsValidHandle())
+	{
+		return FEdGraphPinType();
+	}
+
+	TArray<FPropertyBagPropertyDesc> PropertyDescs = Private::GetCommonPropertyDescs(InBagStructProperty);
+
+	const FProperty* Property = ChildPropertyHandle->GetProperty();
+	if (FPropertyBagPropertyDesc* Desc = PropertyDescs.FindByPredicate([Property](const FPropertyBagPropertyDesc& Desc){ return Desc.CachedProperty == Property; }))
+	{
+		return GetPropertyDescAsPin(*Desc);
+	}
+
+	return FEdGraphPinType();
+}
+
+void PinInfoChanged(TSharedPtr<IPropertyHandle> ChildPropertyHandle, const TSharedPtr<IPropertyHandle>& InBagStructProperty, const TSharedPtr<IPropertyUtilities>& InPropUtils, const FEdGraphPinType& PinType)
+{
+	// The SPinTypeSelector popup might outlive this details view, so bag struct property can be invalid here.
+	if (!InBagStructProperty || !InBagStructProperty->IsValidHandle() || !ChildPropertyHandle || !ChildPropertyHandle->IsValidHandle())
+	{
+		return;
+	}
+
+	Private::ApplyChangesToPropertyDescs(
+		FText::Format(LOCTEXT("OnPropertyTypeChanged", "Changed property type: {0}"), ChildPropertyHandle->GetPropertyDisplayName()),
+		InBagStructProperty,
+		InPropUtils,
+		[&PinType, &ChildPropertyHandle](TArray<FPropertyBagPropertyDesc>& PropertyDescs)
+		{
+			// Find and change struct type
+			const FProperty* Property = ChildPropertyHandle ? ChildPropertyHandle->GetProperty() : nullptr;
+			if (FPropertyBagPropertyDesc* Desc = PropertyDescs.FindByPredicate([Property](const FPropertyBagPropertyDesc& Desc){ return Desc.CachedProperty == Property; }))
+			{
+				SetPropertyDescFromPin(*Desc, PinType);
+			}
+		});
+}
 } // UE::StructUtils::Private
 
+TSharedRef<SWidget> CreateTypeSelectionWidget(TSharedPtr<IPropertyHandle> ChildPropertyHandle, const TSharedPtr<IPropertyHandle>& InBagStructProperty, const TSharedPtr<IPropertyUtilities>& InPropUtils, SPinTypeSelector::ESelectorType SelectorType, const bool bAllowContainers)
+{
+	return SNew(SBox)
+		.HAlign(HAlign_Right)
+		.Padding(FMargin(4, 0))
+		[
+			SNew(STypeSelector, FGetPinTypeTree::CreateLambda(
+				[BagStructProperty = InBagStructProperty](TArray<TSharedPtr<UEdGraphSchema_K2::FPinTypeTreeInfo>>& TypeTree, const ETypeTreeFilter TypeTreeFilter)
+				{
+				UE::StructUtils::Private::GetFilteredVariableTypeTree(BagStructProperty, TypeTree, TypeTreeFilter);
+				}))
+			.TargetPinType_Lambda([ChildPropertyHandle = ChildPropertyHandle, BagStructProperty = InBagStructProperty]()
+			{
+				return Private::GetPinInfo(ChildPropertyHandle, BagStructProperty);
+			})
+			.OnPinTypeChanged_Lambda([ChildPropertyHandle = ChildPropertyHandle, BagStructProperty = InBagStructProperty, PropUtils = InPropUtils](const FEdGraphPinType& PinType)
+			{
+				return Private::PinInfoChanged(ChildPropertyHandle, BagStructProperty, PropUtils, PinType);
+			})
+			.Schema(GetDefault<UPropertyBagSchema>())
+			.bAllowContainers(bAllowContainers)
+			.SelectorType(SelectorType)
+			.TypeTreeFilter(ETypeTreeFilter::None)
+			.Font(IDetailLayoutBuilder::GetDetailFont())
+		];
+}
 } // UE::StructUtils
 
 //----------------------------------------------------------------//
@@ -785,12 +974,28 @@ void CopyPropertyValue(const FPropertyBagPropertyDesc* InSourcePropertyDesc, con
 //  - ChildPropertyHandle a child property of the FInstancedPropertyBag::Value (FInstancedStruct)  
 //----------------------------------------------------------------//
 
-FPropertyBagInstanceDataDetails::FPropertyBagInstanceDataDetails(TSharedPtr<IPropertyHandle> InStructProperty, const TSharedPtr<IPropertyUtilities>& InPropUtils, const bool bInFixedLayout, const bool bInAllowArrays)
+/** Primary constructor. Values passed by parameter struct. */
+FPropertyBagInstanceDataDetails::FPropertyBagInstanceDataDetails(const FPropertyBagInstanceDataDetails::FConstructParams& ConstructParams)
+	: FInstancedStructDataDetails(ConstructParams.BagStructProperty.IsValid() ? ConstructParams.BagStructProperty->GetChildHandle(TEXT("Value")) : nullptr)
+	, BagStructProperty(ConstructParams.BagStructProperty)
+	, PropUtils(ConstructParams.PropUtils)
+	, bAllowContainers(ConstructParams.bAllowContainers)
+	, ChildRowFeatures(ConstructParams.ChildRowFeatures)
+{
+	ensure(UE::StructUtils::Private::IsScriptStruct<FInstancedPropertyBag>(BagStructProperty));
+	ensure(PropUtils != nullptr);
+}
+
+/** For backwards compatibility. */
+FPropertyBagInstanceDataDetails::FPropertyBagInstanceDataDetails(TSharedPtr<IPropertyHandle> InStructProperty, const TSharedPtr<IPropertyUtilities>& InPropUtils, const bool bInFixedLayout, const bool bInAllowContainers)
 	: FInstancedStructDataDetails(InStructProperty.IsValid() ? InStructProperty->GetChildHandle(TEXT("Value")) : nullptr)
 	, BagStructProperty(InStructProperty)
 	, PropUtils(InPropUtils)
-	, bFixedLayout(bInFixedLayout)
-	, bAllowArrays(bInAllowArrays)
+	, bAllowContainers(bInAllowContainers)
+	, ChildRowFeatures(bInFixedLayout ? EPropertyBagChildRowFeatures::Fixed : EPropertyBagChildRowFeatures::Default)
+PRAGMA_DISABLE_DEPRECATION_WARNINGS
+	, bAllowArrays(bInAllowContainers) // For backwards compatibility
+PRAGMA_ENABLE_DEPRECATION_WARNINGS
 {
 	ensure(UE::StructUtils::Private::IsScriptStruct<FInstancedPropertyBag>(BagStructProperty));
 	ensure(PropUtils != nullptr);
@@ -799,73 +1004,157 @@ FPropertyBagInstanceDataDetails::FPropertyBagInstanceDataDetails(TSharedPtr<IPro
 void FPropertyBagInstanceDataDetails::OnChildRowAdded(IDetailPropertyRow& ChildRow)
 {
 	TSharedPtr<SWidget> NameWidget;
-	TSharedPtr<SWidget> ValueWidget;
-	FDetailWidgetRow Row;
-	ChildRow.GetDefaultWidgets(NameWidget, ValueWidget, Row);
+	TSharedPtr<SWidget> PropertyValueWidget;
+	FDetailWidgetRow DetailWidgetRow;
+	ChildRow.GetDefaultWidgets(NameWidget, PropertyValueWidget, DetailWidgetRow);
 
 	TSharedPtr<IPropertyHandle> ChildPropertyHandle = ChildRow.GetPropertyHandle();
 	check(ChildPropertyHandle);
 
-	bool bSupportedType = true;
-	bool bEditable = true;
+	TWeakPtr<FPropertyBagInstanceDataDetails> WeakSelf = SharedThis<FPropertyBagInstanceDataDetails>(this);
+
+	bool bSupportedType = false;
+	bool bEditable = false;
 
 	FText UnsupportedTypeWarning;
 
 	const UPropertyBag* BagStruct = BagStructProperty ? UE::StructUtils::Private::GetCommonBagStruct(BagStructProperty) : nullptr;
-	const FProperty* Property = ChildPropertyHandle->GetProperty();
-	const FPropertyBagPropertyDesc* PropertyDesc = BagStruct ? BagStruct->FindPropertyDescByProperty(Property) : nullptr;
-	if (PropertyDesc)
+	const FProperty* ChildProperty = ChildPropertyHandle->GetProperty();
+	const FPropertyBagPropertyDesc* PropertyDesc = BagStruct ? BagStruct->FindPropertyDescByProperty(ChildProperty) : nullptr;
+
+	PRAGMA_DISABLE_DEPRECATION_WARNINGS
+	const bool bIsFixed = bFixedLayout || ChildRowFeatures == EPropertyBagChildRowFeatures::Fixed;
+	PRAGMA_ENABLE_DEPRECATION_WARNINGS
+
+	// Validate data and check if it's editable
+	if (ChildProperty->HasMetaData(UE::StructUtils::Metadata::HideInDetailPanelsName))
 	{
-		static const FName HideInDetailPanelsName(TEXT("HideInDetailPanel"));
+		ChildRow.Visibility(EVisibility::Collapsed);
+		return;
+	}
+	else if (!ensure(PropertyDesc) || PropertyDesc->ContainerTypes.Num() > 1)
+	{
+		static const FText UnsupportedTypeWarningNestedContainer = LOCTEXT("NestedContainersWarning", "This property type is not supported in the property bag UI.");
 
-		if (PropertyDesc->ContainerTypes.Num() > 1)
-		{
-			static const FText UnsupportedTypeWarningNestedContainer = LOCTEXT("NestedContainersWarning", "This property type is not supported in the property bag UI.");
+		// The property editing for nested containers is not supported.
+		UnsupportedTypeWarning = UnsupportedTypeWarningNestedContainer;
+	}
+	else if ((PropertyDesc->ValueType == EPropertyBagPropertyType::UInt32 || PropertyDesc->ValueType == EPropertyBagPropertyType::UInt64) && !bIsFixed)
+	{
+		static const FText UnsupportedTypeWarningUnsigned = LOCTEXT("UnsignedTypesWarning", "Unsigned types are not supported through the property type selection. If you change the type, you will not be able to change it back.");
 
-			// The property editing for nested containers is not supported.
-			UnsupportedTypeWarning = UnsupportedTypeWarningNestedContainer;
-			bSupportedType = false;
-			// Do not allow to edit the data.
-			bEditable = false;
-		}
-		else
-		{
-			bEditable = BagStructProperty->IsEditable();
+		// Warn that the unsigned types cannot be set via the type selection.
+		UnsupportedTypeWarning = UnsupportedTypeWarningUnsigned;
+	}
+	else // The property is valid - continue.
+	{
+		bEditable = BagStructProperty->IsEditable();
+		bSupportedType = true;
+	}
 
-			if ((PropertyDesc->ValueType == EPropertyBagPropertyType::UInt32 || PropertyDesc->ValueType == EPropertyBagPropertyType::UInt64)
-				&& !bFixedLayout)
+	/*** OVERRIDE RESET TO DEFAULT ACTION FOR BAG OVERRIDES ***/
+	if (HasPropertyOverrides())
+	{
+		TAttribute<bool> EditConditionValue = TAttribute<bool>::CreateLambda(
+			[WeakSelf, ChildPropertyHandle]() -> bool
+			{
+				if (const TSharedPtr<FPropertyBagInstanceDataDetails> Self = WeakSelf.Pin())
+				{
+					return Self->IsPropertyOverridden(ChildPropertyHandle) == EPropertyOverrideState::Yes;
+				}
+				return true;
+			});
+		
+		FOnBooleanValueChanged OnEditConditionChanged = FOnBooleanValueChanged::CreateLambda([WeakSelf, ChildPropertyHandle](bool bNewValue)
+		{
+			if (const TSharedPtr<FPropertyBagInstanceDataDetails> Self = WeakSelf.Pin())
 			{
-				static const FText UnsupportedTypeWarningUnsigned = LOCTEXT("UnsignedTypesWarning", "Unsigned types are not supported throught the property type selection. If you change the type, you will not be able to change it back.");
+				Self->SetPropertyOverride(ChildPropertyHandle, bNewValue);
+			}
+		});
 
-				// Warn that the unsigned types cannot be set via the type selection.
-				UnsupportedTypeWarning = UnsupportedTypeWarningUnsigned;
-				bSupportedType = false;
+		ChildRow.EditCondition(std::move(EditConditionValue), std::move(OnEditConditionChanged));
+
+		FIsResetToDefaultVisible IsResetVisible = FIsResetToDefaultVisible::CreateLambda([WeakSelf](TSharedPtr<IPropertyHandle> PropertyHandle)
+		{
+			if (const TSharedPtr<FPropertyBagInstanceDataDetails> Self = WeakSelf.Pin())
+			{
+				return !Self->IsDefaultValue(PropertyHandle);
 			}
-			else if (Property->HasMetaData(HideInDetailPanelsName))
+			return false;
+		});
+		FResetToDefaultHandler ResetHandler = FResetToDefaultHandler::CreateLambda([WeakSelf](TSharedPtr<IPropertyHandle> PropertyHandle)
+		{
+			if (const TSharedPtr<FPropertyBagInstanceDataDetails> Self = WeakSelf.Pin())
 			{
-				ChildRow.Visibility(EVisibility::Hidden);
-				return;
+				Self->ResetToDefault(PropertyHandle);
 			}
-		}
-	}
+		});
+		FResetToDefaultOverride ResetOverride = FResetToDefaultOverride::Create(IsResetVisible, ResetHandler);
 
-	ChildRow.IsEnabled(bEditable);
+		ChildRow.OverrideResetToDefault(ResetOverride);
+	}
 
-	if (!bFixedLayout)
+	if (!bIsFixed)
 	{
-		// Inline editable name
-		TSharedPtr<SInlineEditableTextBlock> InlineWidget = SNew(SInlineEditableTextBlock)
+		/*** BUILD PROPERTY NAME WIDGET ***/
+		TSharedRef<SHorizontalBox> PropertyDetailsWidget = SNew(SHorizontalBox)
+			+ SHorizontalBox::Slot()
+			.HAlign(HAlign_Right)
+			.Padding(3, 0)
+			.AutoWidth()
+			[
+				SNew(SBox)
+				.HAlign(HAlign_Center)
+				.VAlign(VAlign_Center)
+				[
+					SNew(SImage)
+					.ToolTipText_Lambda([UnsupportedTypeWarning, bSupportedType]()
+					{
+						return !bSupportedType ? UnsupportedTypeWarning : LOCTEXT("MissingType", "The property is missing type. The Struct, Enum, or Object may have been removed.");
+					})
+					.Visibility_Lambda([ChildPropertyHandle, bSupportedType]()
+					{
+						return UE::StructUtils::Private::HasMissingType(ChildPropertyHandle) || !bSupportedType ? EVisibility::Visible : EVisibility::Collapsed;
+					})
+					.DesiredSizeOverride(FVector2D(12))
+					.ColorAndOpacity(FLinearColor(1.f, 0.8f, 0.f, 1.f))
+					.Image(FAppStyle::GetBrush("Icons.Error"))
+				]
+			];
+
+		if (EnumHasAnyFlags(ChildRowFeatures, EPropertyBagChildRowFeatures::CompactTypeSelector))
+		{
+			PropertyDetailsWidget->AddSlot()
+				.HAlign(HAlign_Left)
+				.Padding(1, 0)
+				.AutoWidth()
+				[
+					UE::StructUtils::CreateTypeSelectionWidget(ChildPropertyHandle, BagStructProperty, PropUtils, SPinTypeSelector::ESelectorType::Compact, bAllowContainers)
+				];
+		}
+
+		// Choose the name text box, editable inline or not.
+		TSharedPtr<SInlineEditableTextBlock> EditableInlineNameWidget = SNew(SInlineEditableTextBlock)
+			.IsReadOnly(!EnumHasAnyFlags(ChildRowFeatures, EPropertyBagChildRowFeatures::Renaming))
 			.Font(IDetailLayoutBuilder::GetDetailFont())
 			.MultiLine(false)
+			.OverflowPolicy(ETextOverflowPolicy::Ellipsis)
 			.Text_Lambda([ChildPropertyHandle]()
 			{
 				return FText::FromName(UE::StructUtils::Private::GetPropertyNameSafe(ChildPropertyHandle));
 			})
 			.OnVerifyTextChanged_Lambda([BagStructProperty = BagStructProperty, ChildPropertyHandle](const FText& InText, FText& OutErrorMessage)
 			{
-				const FName NewName = UE::StructUtils::Private::GetValidPropertyName(InText.ToString());
+				if (InText.IsEmpty())
+				{
+					OutErrorMessage = LOCTEXT("EmptyName", "Name is empty");
+					return false;
+				}
+
+				const FName NewName = FInstancedPropertyBag::SanitizePropertyName(InText.ToString());
 				const FName OldName = UE::StructUtils::Private::GetPropertyNameSafe(ChildPropertyHandle);
-				bool bResult = UE::StructUtils::Private::IsUniqueName(NewName, OldName, BagStructProperty);
+				const bool bResult = UE::StructUtils::Private::IsUniqueName(NewName, OldName, BagStructProperty);
 				if (!bResult)
 				{
 					OutErrorMessage = LOCTEXT("MustBeUniqueName", "Property must have unique name");
@@ -879,7 +1168,7 @@ void FPropertyBagInstanceDataDetails::OnChildRowAdded(IDetailPropertyRow& ChildR
 					return;
 				}
 
-				const FName NewName = UE::StructUtils::Private::GetValidPropertyName(InNewText.ToString());
+				const FName NewName = FInstancedPropertyBag::SanitizePropertyName(InNewText.ToString());
 				const FName OldName = UE::StructUtils::Private::GetPropertyNameSafe(ChildPropertyHandle);
 				if (!UE::StructUtils::Private::IsUniqueName(NewName, OldName, BagStructProperty))
 				{
@@ -887,141 +1176,247 @@ void FPropertyBagInstanceDataDetails::OnChildRowAdded(IDetailPropertyRow& ChildR
 				}
 
 				UE::StructUtils::Private::ApplyChangesToPropertyDescs(
-					LOCTEXT("OnPropertyNameChanged", "Change Property Name"), BagStructProperty, PropUtils,
+					FText::Format(LOCTEXT("OnPropertyNameChanged", "Change property name: {0} -> {1}"),
+						FText::FromName(OldName),
+						FText::FromName(NewName)),
+					BagStructProperty,
+					PropUtils,
 					[&NewName, &ChildPropertyHandle](TArray<FPropertyBagPropertyDesc>& PropertyDescs)
 					{
 						const FProperty* Property = ChildPropertyHandle->GetProperty();
-						if (FPropertyBagPropertyDesc* Desc = PropertyDescs.FindByPredicate([Property](const FPropertyBagPropertyDesc& Desc){ return Desc.CachedProperty == Property; }))
+						if (FPropertyBagPropertyDesc* Desc = PropertyDescs.FindByPredicate([Property](const FPropertyBagPropertyDesc& Desc) { return Desc.CachedProperty == Property; }))
 						{
 							Desc->Name = NewName;
 						}
 					});
 			});
 
-		NameWidget = SNew(SComboButton)
-			.OnGetMenuContent(this, &FPropertyBagInstanceDataDetails::OnPropertyNameContent, ChildRow.GetPropertyHandle(), InlineWidget)
-			.ContentPadding(2)
-			.ButtonStyle(FAppStyle::Get(), "HoverHintOnly")
-			.ButtonContent()
+		// Deprecated in 5.6 - the combo button on the name widget will be removed in favor of the new drop-down menu.
+		if (EnumHasAnyFlags(ChildRowFeatures, EPropertyBagChildRowFeatures::Deprecated))
+		{
+			// Add the widget to the property bar.
+			PropertyDetailsWidget->AddSlot()
+				.HAlign(HAlign_Left)
+				.AutoWidth()
+				.Padding(0, 0, 3, 0)
+				[
+					SNew(SBox)
+					.Content()
+					[
+						SNew(SComboButton)
+						.MenuContent()
+						[
+							PRAGMA_DISABLE_DEPRECATION_WARNINGS
+							OnPropertyNameContent(ChildPropertyHandle, EditableInlineNameWidget)
+							PRAGMA_ENABLE_DEPRECATION_WARNINGS
+						]
+						.ContentPadding(FMargin(0, 0, 2, 0))
+						.ButtonStyle(FAppStyle::Get(), "SimpleButton")
+						.ButtonContent()
+						[
+							EditableInlineNameWidget.ToSharedRef()
+						]
+					]
+				];
+		}
+		else // No deprecated combo box. Just add the name.
+		{
+			PropertyDetailsWidget->AddSlot()
+				.HAlign(HAlign_Left)
+				.AutoWidth()
+				.Padding(0, 0, 3, 0)
+				[
+					EditableInlineNameWidget.ToSharedRef()
+				];
+		}
+
+		// Extendable spacer between the name and the drop-down
+		PropertyDetailsWidget->AddSlot()
+			.FillWidth(1)
 			[
-				InlineWidget.ToSharedRef()
+				SNew(SSpacer)
+				.Size(1)
 			];
-	}
 
-	ChildRow
-		.CustomWidget(/*bShowChildren*/true)
-		.NameContent()
-		[
-			SNew(SHorizontalBox)
-			// Error icon
-			+ SHorizontalBox::Slot()
-			.AutoWidth()
-			.HAlign(HAlign_Right)
-			.VAlign(VAlign_Center)
-			.Padding(0,0,2,0)
-			[
-				SNew(SBox)
-				.WidthOverride(12)
-				.HeightOverride(12)
+		if (EnumHasAnyFlags(ChildRowFeatures, EPropertyBagChildRowFeatures::AccessSpecifierButton))
+		{
+			using namespace UE::StructUtils::Metadata;
+			PropertyDetailsWidget->AddSlot()
+				.HAlign(HAlign_Right)
+				.VAlign(VAlign_Center)
+				.AutoWidth()
+				.Padding(1, 0)
 				[
-					SNew(SImage)
-					.ToolTipText_Lambda([UnsupportedTypeWarning, bSupportedType]()
+					SNew(SButton)
+					.HAlign(HAlign_Center)
+					.VAlign(VAlign_Center)
+					.ButtonStyle(FAppStyle::Get(), "SimpleButton")
+					.ToolTipText(LOCTEXT("SetAccessSpecifier", "Set the access specifier on the property to Public or Private."))
+					.OnClicked_Lambda([ChildPropertyHandle, BagStructProperty = BagStructProperty, PropUtils = PropUtils]()
 					{
-						return !bSupportedType
-							? UnsupportedTypeWarning
-							: LOCTEXT("MissingType", "The property is missing type. The Struct, Enum, or Object may have been removed.");
+						FProperty* Property = ChildPropertyHandle->GetProperty();
+						UE::StructUtils::Private::ApplyChangesToPropertyDescs(
+							LOCTEXT("OnPropertyAccessSpecifierChanged", "Set access specifier."),
+							BagStructProperty,
+							PropUtils,
+							[Property](TArray<FPropertyBagPropertyDesc>& PropertyDescs)
+							{
+								if (FPropertyBagPropertyDesc* Desc = PropertyDescs.FindByPredicate([Property](const FPropertyBagPropertyDesc& Desc) { return Desc.CachedProperty == Property; }))
+								{
+									const bool bIsPrivate = Property->HasAnyPropertyFlags(CPF_NativeAccessSpecifierPrivate | CPF_NativeAccessSpecifierProtected);
+									Desc->PropertyFlags &= ~CPF_NativeAccessSpecifiers;
+									if (bIsPrivate)
+									{
+										Desc->PropertyFlags |= CPF_NativeAccessSpecifierPublic;
+									}
+									else
+									{
+										Desc->PropertyFlags |= CPF_NativeAccessSpecifierPrivate;
+									}
+								}
+							});
+
+						return FReply::Handled();
 					})
-					.Visibility_Lambda([ChildPropertyHandle, bSupportedType]()
-					{
-						return UE::StructUtils::Private::HasMissingType(ChildPropertyHandle) || !bSupportedType ? EVisibility::Visible : EVisibility::Collapsed;
-					})
-					.Image(FAppStyle::GetBrush("Icons.Error"))
-				]
-			]
-			// Name
-			+SHorizontalBox::Slot()
-			.AutoWidth()
-			[
-				NameWidget.ToSharedRef()
-			]
-		]
-		.ValueContent()
-		[
-			ValueWidget.ToSharedRef()
-		];
+					.ButtonStyle(FAppStyle::Get(), "SimpleButton")
+					[
+						SNew(SImage)
+						.DesiredSizeOverride(FVector2D(16))
+						.ColorAndOpacity(FSlateColor::UseForeground())
+						.Image_Lambda([ChildPropertyHandle]()
+						{
+							check(ChildPropertyHandle);
+							if (const FProperty* Property = ChildPropertyHandle->GetProperty())
+							{
+								// For now, treat protected as private. TODO: Add toggle for protected.
+								if (Property->HasAnyPropertyFlags(CPF_NativeAccessSpecifierPrivate | CPF_NativeAccessSpecifierProtected))
+								{
+									return FAppStyle::Get().GetBrush("Icons.Visible");
+								}
+							}
+
+							return FAppStyle::Get().GetBrush("Icons.Hidden");
+						})
+					]
+				];
+		}
 
-	TWeakPtr<FPropertyBagInstanceDataDetails> WeakSelf = SharedThis<FPropertyBagInstanceDataDetails>(this);
+		// Check drop-down is enabled and at least one option as well.
+		if (EnumHasAnyFlags(ChildRowFeatures, EPropertyBagChildRowFeatures::DropDownMenuButton)
+			&& EnumHasAnyFlags(ChildRowFeatures, EPropertyBagChildRowFeatures::AllMenuOptions))
+		{
+			static constexpr bool bInShouldCloseWindowAfterMenuSelection = true;
+			FMenuBuilder MenuBuilder(bInShouldCloseWindowAfterMenuSelection, nullptr);
 
-	if (PropertyDesc)
-	{
-		ChildRow.DragDropHandler(MakeShared<FPropertyBagDetailsDragDropHandler>(*PropertyDesc, FOnPropertyBagDetailsRowDropOp::CreateLambda(
-			[WeakSelf, PropertyDesc = *PropertyDesc](const FPropertyBagPropertyDesc& DroppedPropertyDesc, EItemDropZone DropZone) -> FReply
+			if (EnumHasAnyFlags(ChildRowFeatures, EPropertyBagChildRowFeatures::Menu_TypeSelector))
 			{
-				const TSharedPtr<FPropertyBagInstanceDataDetails> DetailsSP = WeakSelf.Pin();
-				const UPropertyBag* ChildBagStruct = DetailsSP ? UE::StructUtils::Private::GetCommonBagStruct(DetailsSP->BagStructProperty) : nullptr;
-				// Validate these properties are still part of the bag.
-				if (!ChildBagStruct || !ChildBagStruct->FindPropertyDescByProperty(PropertyDesc.CachedProperty) || !ChildBagStruct->FindPropertyDescByProperty(DroppedPropertyDesc.CachedProperty))
-				{
-					return FReply::Unhandled();
-				}
+				MenuBuilder.AddWidget(UE::StructUtils::CreateTypeSelectionWidget(ChildPropertyHandle,
+						BagStructProperty, PropUtils,
+						SPinTypeSelector::ESelectorType::Full,
+						bAllowContainers),
+					/*InLabel=*/FText::GetEmpty());
+				MenuBuilder.AddMenuSeparator();
+			}
+
+			// Must have property renaming enabled or the editable inline widget will be invalid.
+			if (EnumHasAllFlags(ChildRowFeatures, EPropertyBagChildRowFeatures::Renaming | EPropertyBagChildRowFeatures::Menu_Rename))
+			{
+				MenuBuilder.AddMenuEntry(
+					LOCTEXT("Rename", "Rename"),
+					LOCTEXT("Rename_ToolTip", "Rename property"),
+					FSlateIcon(FAppStyle::GetAppStyleSetName(), "Icons.Edit"),
+					FUIAction(FExecuteAction::CreateLambda([NameWidget = EditableInlineNameWidget]()
+					{
+						if (NameWidget)
+						{
+							NameWidget->EnterEditingMode();
+						}
+					})));
+			}
 
-				EPropertyBagAlterationResult Result = EPropertyBagAlterationResult::InternalError;
+			if (EnumHasAllFlags(ChildRowFeatures, EPropertyBagChildRowFeatures::Deletion | EPropertyBagChildRowFeatures::Menu_Delete))
+			{
+				MenuBuilder.AddMenuEntry(
+					LOCTEXT("Remove", "Remove"),
+					LOCTEXT("Remove_ToolTip", "Remove property"),
+					FSlateIcon(FAppStyle::GetAppStyleSetName(), "Icons.Delete"),
+					FUIAction(FExecuteAction::CreateLambda([BagStructProperty = BagStructProperty, PropUtils = PropUtils, ChildPropertyHandle]()
+					{
+						UE::StructUtils::Private::DeleteProperty(BagStructProperty, ChildPropertyHandle, PropUtils);
+					}))
+				);
+			}
 
-				DetailsSP->BagStructProperty->EnumerateRawData([DropZone, &PropertyDesc, &DroppedPropertyDesc, &Result](void* RawData, const int32 /*DataIndex*/, const int32 /*NumData*/)
+					/*** Drop-down arrow menu. ***/
+			PropertyDetailsWidget->AddSlot()
+				.HAlign(HAlign_Right)
+				.AutoWidth()
+				.Padding(0, 0, 5, 0)
+				[
+					SNew(SBox)
+					.Content()
+					[
+						SNew(SComboButton)
+						.MenuContent()
+						[
+							MenuBuilder.MakeWidget()
+						]
+						.HasDownArrow(true)
+						.ButtonStyle(FAppStyle::Get(), "SimpleButton")
+						.ButtonContent()
+						[
+							SNullWidget::NullWidget
+						]
+					]
+				];
+		}
+
+		/*** DRAG AND DROP HANDLER ***/
+		if (EnumHasAnyFlags(ChildRowFeatures, EPropertyBagChildRowFeatures::DragAndDrop) && PropertyDesc != nullptr)
+		{
+			ChildRow.DragDropHandler(MakeShared<FPropertyBagDetailsDragDropHandler>(*PropertyDesc, FOnPropertyBagDetailsRowDropOp::CreateLambda(
+				[WeakSelf = SharedThis<FPropertyBagInstanceDataDetails>(this).ToWeakPtr(), PropertyDesc = *PropertyDesc](const FPropertyBagPropertyDesc& DroppedPropertyDesc, EItemDropZone DropZone) -> FReply
 				{
-					if (FInstancedPropertyBag* PropertyBag = RawData ? static_cast<FInstancedPropertyBag*>(RawData) : nullptr)
+					const TSharedPtr<FPropertyBagInstanceDataDetails> DetailsSP = WeakSelf.Pin();
+					const UPropertyBag* ChildBagStruct = DetailsSP ? UE::StructUtils::Private::GetCommonBagStruct(DetailsSP->BagStructProperty) : nullptr;
+					// Validate these properties are still part of the bag.
+					if (!ChildBagStruct || !ChildBagStruct->FindPropertyDescByProperty(PropertyDesc.CachedProperty) || !ChildBagStruct->FindPropertyDescByProperty(DroppedPropertyDesc.CachedProperty))
 					{
-						Result = PropertyBag->ReorderProperty(DroppedPropertyDesc.Name, PropertyDesc.Name, DropZone == EItemDropZone::AboveItem);
+						return FReply::Unhandled();
 					}
 
-					return true;
-				});
+					EPropertyBagAlterationResult Result = EPropertyBagAlterationResult::InternalError;
 
-				return Result == EPropertyBagAlterationResult::Success ? FReply::Handled() : FReply::Unhandled();
-			})));
-	}
+					DetailsSP->BagStructProperty->EnumerateRawData([DropZone, &PropertyDesc, &DroppedPropertyDesc, &Result](void* RawData, const int32 /*DataIndex*/, const int32 /*NumData*/)
+					{
+						if (FInstancedPropertyBag* PropertyBag = RawData ? static_cast<FInstancedPropertyBag*>(RawData) : nullptr)
+						{
+							Result = PropertyBag->ReorderProperty(DroppedPropertyDesc.Name, PropertyDesc.Name, DropZone == EItemDropZone::AboveItem);
+						}
 
-	if (HasPropertyOverrides())
-	{
-		TAttribute<bool> EditConditionValue = TAttribute<bool>::CreateLambda(
-			[WeakSelf, ChildPropertyHandle]() -> bool
-			{
-				if (const TSharedPtr<FPropertyBagInstanceDataDetails> Self = WeakSelf.Pin())
-				{
-					return Self->IsPropertyOverridden(ChildPropertyHandle) == EPropertyOverrideState::Yes;
-				}
-				return true;
-			});
-		
-		FOnBooleanValueChanged OnEditConditionChanged = FOnBooleanValueChanged::CreateLambda([WeakSelf, ChildPropertyHandle](bool bNewValue)
-		{
-			if (const TSharedPtr<FPropertyBagInstanceDataDetails> Self = WeakSelf.Pin())
-			{
-				Self->SetPropertyOverride(ChildPropertyHandle, bNewValue);
-			}
-		});
+						return true;
+					});
 
-		ChildRow.EditCondition(std::move(EditConditionValue), std::move(OnEditConditionChanged));
+					return Result == EPropertyBagAlterationResult::Success ? FReply::Handled() : FReply::Unhandled();
+				})));
+		}
 
-		FIsResetToDefaultVisible IsResetVisible = FIsResetToDefaultVisible::CreateLambda([WeakSelf](TSharedPtr<IPropertyHandle> PropertyHandle)
-		{
-			if (const TSharedPtr<FPropertyBagInstanceDataDetails> Self = WeakSelf.Pin())
-			{
-				return !Self->IsDefaultValue(PropertyHandle);
-			}
-			return false;
-		});
-		FResetToDefaultHandler ResetHandler = FResetToDefaultHandler::CreateLambda([WeakSelf](TSharedPtr<IPropertyHandle> PropertyHandle)
-		{
-			if (const TSharedPtr<FPropertyBagInstanceDataDetails> Self = WeakSelf.Pin())
-			{
-				Self->ResetToDefault(PropertyHandle);
-			}
-		});
-		FResetToDefaultOverride ResetOverride = FResetToDefaultOverride::Create(IsResetVisible, ResetHandler);
+		// Update the name widget with our new property details composition.
+		NameWidget = PropertyDetailsWidget;
+	} // Fixed layouts will pass through and use the default widgets.
 
-		ChildRow.OverrideResetToDefault(ResetOverride);
-	}
+	ChildRow
+		.IsEnabled(bEditable)
+		.CustomWidget(/*bShowChildren*/true)
+		.NameContent()
+		.HAlign(HAlign_Fill)
+		[
+			NameWidget.ToSharedRef()
+		]
+		.ValueContent()
+		[
+			PropertyValueWidget.ToSharedRef()
+		];
 }
 
 FPropertyBagInstanceDataDetails::EPropertyOverrideState FPropertyBagInstanceDataDetails::IsPropertyOverridden(TSharedPtr<IPropertyHandle> ChildPropertyHandle) const
@@ -1190,161 +1585,6 @@ TSharedRef<SWidget> FPropertyBagInstanceDataDetails::OnPropertyNameContent(TShar
 	constexpr bool bInShouldCloseWindowAfterMenuSelection = true;
 	FMenuBuilder MenuBuilder(bInShouldCloseWindowAfterMenuSelection, nullptr);
 
-	auto GetFilteredVariableTypeTree = [BagStructProperty = BagStructProperty](TArray<TSharedPtr<UEdGraphSchema_K2::FPinTypeTreeInfo>>& TypeTree, ETypeTreeFilter TypeTreeFilter)
-	{
-		// The SPinTypeSelector popup might outlive this details view, so bag struct property can be invalid here.
-    	if (!BagStructProperty || !BagStructProperty->IsValidHandle())
-    	{
-    		return;
-    	}
-
-		UFunction* IsPinTypeAcceptedFunc = nullptr;
-		UObject* IsPinTypeAcceptedTarget = nullptr;
-		if (UE::StructUtils::Private::FindUserFunction(BagStructProperty, UE::StructUtils::Metadata::IsPinTypeAcceptedName, IsPinTypeAcceptedFunc, IsPinTypeAcceptedTarget))
-		{
-			check(IsPinTypeAcceptedFunc && IsPinTypeAcceptedTarget);
-
-			// We need to make sure the signature matches perfectly: bool(FEdGraphPinType, bool)
-			bool bFuncIsValid = UE::StructUtils::Private::ValidateFunctionSignature<bool, FEdGraphPinType, bool>(IsPinTypeAcceptedFunc);
-			if (!ensureMsgf(bFuncIsValid, TEXT("[%s] Function %s does not have the right signature."), *UE::StructUtils::Metadata::IsPinTypeAcceptedName.ToString(), *IsPinTypeAcceptedFunc->GetName()))
-			{
-				return;
-			}
-		}
-
-		auto IsPinTypeAccepted = [IsPinTypeAcceptedFunc, IsPinTypeAcceptedTarget](const FEdGraphPinType& InPinType, bool bInIsChild) -> bool
-		{
-			if (IsPinTypeAcceptedFunc && IsPinTypeAcceptedTarget)
-			{
-				const TValueOrError<bool, void> bIsValid = UE::StructUtils::Private::CallFunc<bool>(IsPinTypeAcceptedTarget, IsPinTypeAcceptedFunc, InPinType, bInIsChild);
-				return bIsValid.HasValue() && bIsValid.GetValue();
-			}
-			else
-			{
-				return true;
-			}
-		};
-
-		check(GetDefault<UEdGraphSchema_K2>());
-		TArray<TSharedPtr<UEdGraphSchema_K2::FPinTypeTreeInfo>> TempTypeTree;
-		GetDefault<UPropertyBagSchema>()->GetVariableTypeTree(TempTypeTree, TypeTreeFilter);
-
-		// Filter
-		for (TSharedPtr<UEdGraphSchema_K2::FPinTypeTreeInfo>& PinType : TempTypeTree)
-		{
-			if (!PinType.IsValid() || !IsPinTypeAccepted(PinType->GetPinType(/*bForceLoadSubCategoryObject*/false), /*bInIsChild=*/ false))
-			{
-				continue;
-			}
-
-			for (int32 ChildIndex = 0; ChildIndex < PinType->Children.Num(); )
-			{
-				TSharedPtr<UEdGraphSchema_K2::FPinTypeTreeInfo> Child = PinType->Children[ChildIndex];
-				if (Child.IsValid())
-				{
-					const FEdGraphPinType& ChildPinType = Child->GetPinType(/*bForceLoadSubCategoryObject*/false);
-
-					if (!UE::StructUtils::Private::CanHaveMemberVariableOfType(ChildPinType) || !IsPinTypeAccepted(ChildPinType, /*bInIsChild=*/ true))
-					{
-						PinType->Children.RemoveAt(ChildIndex);
-						continue;
-					}
-				}
-				++ChildIndex;
-			}
-
-			TypeTree.Add(PinType);
-		}
-	};
-
-	auto GetPinInfo = [BagStructProperty = BagStructProperty, ChildPropertyHandle]()
-	{
-		// The SPinTypeSelector popup might outlive this details view, so bag struct property can be invalid here.
-		if (!BagStructProperty || !BagStructProperty->IsValidHandle() || !ChildPropertyHandle || !ChildPropertyHandle->IsValidHandle())
-		{
-			return FEdGraphPinType();
-		}
-		
-		TArray<FPropertyBagPropertyDesc> PropertyDescs = UE::StructUtils::Private::GetCommonPropertyDescs(BagStructProperty);
-
-		const FProperty* Property = ChildPropertyHandle->GetProperty();
-		if (FPropertyBagPropertyDesc* Desc = PropertyDescs.FindByPredicate([Property](const FPropertyBagPropertyDesc& Desc){ return Desc.CachedProperty == Property; }))
-		{
-			return UE::StructUtils::GetPropertyDescAsPin(*Desc);
-		}
-	
-		return FEdGraphPinType();
-	};
-
-	auto PinInfoChanged = [BagStructProperty = BagStructProperty, PropUtils = PropUtils, ChildPropertyHandle](const FEdGraphPinType& PinType)
-	{
-		// The SPinTypeSelector popup might outlive this details view, so bag struct property can be invalid here.
-		if (!BagStructProperty || !BagStructProperty->IsValidHandle() || !ChildPropertyHandle || !ChildPropertyHandle->IsValidHandle())
-		{
-			return;
-		}
-		
-		UE::StructUtils::Private::ApplyChangesToPropertyDescs(
-			LOCTEXT("OnPropertyTypeChanged", "Change Property Type"), BagStructProperty, PropUtils,
-			[&PinType, &ChildPropertyHandle](TArray<FPropertyBagPropertyDesc>& PropertyDescs)
-			{
-				// Find and change struct type
-				const FProperty* Property = ChildPropertyHandle ? ChildPropertyHandle->GetProperty() : nullptr;
-				if (FPropertyBagPropertyDesc* Desc = PropertyDescs.FindByPredicate([Property](const FPropertyBagPropertyDesc& Desc){ return Desc.CachedProperty == Property; }))
-				{
-					UE::StructUtils::SetPropertyDescFromPin(*Desc, PinType);
-				}
-			});
-	};
-
-	auto RemoveProperty = [BagStructProperty = BagStructProperty, PropUtils = PropUtils, ChildPropertyHandle]()
-	{
-		if (!BagStructProperty || !BagStructProperty->IsValidHandle() || !ChildPropertyHandle || !ChildPropertyHandle->IsValidHandle())
-		{
-			return;
-		}
-
-		// Extra check provided by the user to cancel a remove action. Useful to provide the user a possibility to cancel the action if
-		// the given property is in use elsewhere.
-		UFunction* CanRemovePropertyFunc = nullptr;
-		UObject* CanRemovePropertyTarget = nullptr;
-		if (UE::StructUtils::Private::FindUserFunction(BagStructProperty, UE::StructUtils::Metadata::CanRemovePropertyName, CanRemovePropertyFunc, CanRemovePropertyTarget))
-		{
-			check(CanRemovePropertyFunc && CanRemovePropertyTarget);
-
-			FName PropertyName = ChildPropertyHandle->GetProperty()->GetFName();
-			const UPropertyBag* PropertyBag = UE::StructUtils::Private::GetCommonBagStruct(BagStructProperty);
-			const FPropertyBagPropertyDesc* PropertyDesc = PropertyBag ? PropertyBag->FindPropertyDescByName(PropertyName) : nullptr;
-
-			if (!PropertyDesc)
-			{
-				return;
-			}
-
-			// We need to make sure the signature matches perfectly: bool(FGuid, FName)s
-			bool bFuncIsValid = UE::StructUtils::Private::ValidateFunctionSignature<bool, FGuid, FName>(CanRemovePropertyFunc);
-			if(!ensureMsgf(bFuncIsValid, TEXT("[%s] Function %s does not have the right signature."), *UE::StructUtils::Metadata::CanRemovePropertyName.ToString(), *CanRemovePropertyFunc->GetName()))
-			{
-				return;
-			}
-
-			const TValueOrError<bool, void> bCanRemove = UE::StructUtils::Private::CallFunc<bool>(CanRemovePropertyTarget, CanRemovePropertyFunc, PropertyDesc->ID, PropertyDesc->Name);
-
-			if (bCanRemove.HasError() || !bCanRemove.GetValue())
-			{
-				return;
-			}
-		}
-
-		UE::StructUtils::Private::ApplyChangesToPropertyDescs(
-			LOCTEXT("OnPropertyRemoved", "Remove Property"), BagStructProperty, PropUtils,
-			[&ChildPropertyHandle](TArray<FPropertyBagPropertyDesc>& PropertyDescs)
-			{
-				const FProperty* Property = ChildPropertyHandle ? ChildPropertyHandle->GetProperty() : nullptr;
-				PropertyDescs.RemoveAll([Property](const FPropertyBagPropertyDesc& Desc){ return Desc.CachedProperty == Property; });
-			});
-	};
-
 	auto MoveProperty = [BagStructProperty = BagStructProperty, PropUtils = PropUtils, ChildPropertyHandle](const int32 Delta)
 	{
 		if (!BagStructProperty || !BagStructProperty->IsValidHandle() || !ChildPropertyHandle || !ChildPropertyHandle->IsValidHandle())
@@ -1370,19 +1610,16 @@ TSharedRef<SWidget> FPropertyBagInstanceDataDetails::OnPropertyNameContent(TShar
 		});
 	};
 
-
 	MenuBuilder.AddWidget(
 		SNew(SBox)
 		.HAlign(HAlign_Right)
 		.Padding(FMargin(12, 0, 12, 0))
 		[
-			SNew(SPinTypeSelector, FGetPinTypeTree::CreateLambda(GetFilteredVariableTypeTree))
-				.TargetPinType_Lambda(GetPinInfo)
-				.OnPinTypeChanged_Lambda(PinInfoChanged)
-				.Schema(GetDefault<UPropertyBagSchema>())
-				.bAllowArrays(bAllowArrays)
-				.TypeTreeFilter(ETypeTreeFilter::None)
-				.Font( IDetailLayoutBuilder::GetDetailFont() )
+			UE::StructUtils::CreateTypeSelectionWidget(ChildPropertyHandle,
+				BagStructProperty,
+				PropUtils,
+				SPinTypeSelector::ESelectorType::Full,
+				bAllowContainers)
 		],
 		FText::GetEmpty());
 	
@@ -1399,7 +1636,10 @@ TSharedRef<SWidget> FPropertyBagInstanceDataDetails::OnPropertyNameContent(TShar
 		LOCTEXT("Remove", "Remove"),
 		LOCTEXT("Remove_ToolTip", "Remove property"),
 		FSlateIcon(FAppStyle::GetAppStyleSetName(), "Icons.Delete"),
-		FUIAction(FExecuteAction::CreateLambda(RemoveProperty))
+		FUIAction(FExecuteAction::CreateLambda([BagStructProperty = BagStructProperty, PropUtils = PropUtils, ChildPropertyHandle]()
+		{
+			UE::StructUtils::Private::DeleteProperty(BagStructProperty, ChildPropertyHandle, PropUtils);
+		}))
 	);
 	
 	MenuBuilder.AddSeparator();
@@ -1421,7 +1661,6 @@ TSharedRef<SWidget> FPropertyBagInstanceDataDetails::OnPropertyNameContent(TShar
 	return MenuBuilder.MakeWidget();
 }
 
-
 ////////////////////////////////////
 
 TSharedRef<IPropertyTypeCustomization> FPropertyBagDetails::MakeInstance()
@@ -1435,19 +1674,17 @@ void FPropertyBagDetails::CustomizeHeader(TSharedRef<IPropertyHandle> StructProp
 	
 	StructProperty = StructPropertyHandle;
 	check(StructProperty);
-	
-	static const FName NAME_ShowOnlyInnerProperties = "ShowOnlyInnerProperties";
-	
+
 	if (const FProperty* MetaDataProperty = StructProperty->GetMetaDataProperty())
 	{
 		bFixedLayout = MetaDataProperty->HasMetaData(UE::StructUtils::Metadata::FixedLayoutName);
 
-		bAllowArrays = MetaDataProperty->HasMetaData(UE::StructUtils::Metadata::AllowArraysName)
-			? MetaDataProperty->GetBoolMetaData(UE::StructUtils::Metadata::AllowArraysName)
+		bAllowContainers = MetaDataProperty->HasMetaData(UE::StructUtils::Metadata::AllowContainersName)
+			? MetaDataProperty->GetBoolMetaData(UE::StructUtils::Metadata::AllowContainersName)
 			: true;
 
 		// Don't show the header if ShowOnlyInnerProperties is set
-		if (MetaDataProperty->HasMetaData(NAME_ShowOnlyInnerProperties))
+		if (MetaDataProperty->HasMetaData(UE::StructUtils::Metadata::ShowOnlyInnerPropertiesName))
 		{
 			return;
 		}
@@ -1463,6 +1700,20 @@ void FPropertyBagDetails::CustomizeHeader(TSharedRef<IPropertyHandle> StructProp
 				}
 			}
 		}
+
+		// Load the feature set by the metadata set on the FPropertyBag. Can only accept explicit enum values currently.
+		// TODO: Enable the option to parse a bitflag expression from string. I.E. 'Renaming | DropDownMenuButton | AllMenuOptions'
+		if (MetaDataProperty->HasMetaData(UE::StructUtils::Metadata::ChildRowFeaturesName))
+		{
+			if (const UEnum* Enum = StaticEnum<EPropertyBagChildRowFeatures>())
+			{
+				const int32 EnumIndex = Enum->GetIndexByNameString(MetaDataProperty->GetMetaData(UE::StructUtils::Metadata::ChildRowFeaturesName));
+				if (EnumIndex != INDEX_NONE)
+				{
+					ChildRowFeatures = static_cast<EPropertyBagChildRowFeatures>(Enum->GetValueByIndex(EnumIndex));
+				}
+			}
+		}
 	}
 
 	TSharedPtr<SWidget> ValueWidget = SNullWidget::NullWidget;
@@ -1486,8 +1737,16 @@ void FPropertyBagDetails::CustomizeHeader(TSharedRef<IPropertyHandle> StructProp
 
 void FPropertyBagDetails::CustomizeChildren(TSharedRef<IPropertyHandle> StructPropertyHandle, IDetailChildrenBuilder& StructBuilder, IPropertyTypeCustomizationUtils& StructCustomizationUtils)
 {
+	FPropertyBagInstanceDataDetails::FConstructParams Params
+	{
+		.BagStructProperty = StructProperty,
+		.PropUtils = PropUtils,
+		.bAllowContainers = bAllowContainers,
+		.ChildRowFeatures = ChildRowFeatures,
+	};
+
 	// Show the Value (FInstancedStruct) as child rows.
-	const TSharedRef<FPropertyBagInstanceDataDetails> InstanceDetails = MakeShareable(new FPropertyBagInstanceDataDetails(StructProperty, PropUtils, bFixedLayout, bAllowArrays));
+	const TSharedRef<FPropertyBagInstanceDataDetails> InstanceDetails = MakeShared<FPropertyBagInstanceDataDetails>(Params);
 	StructBuilder.AddCustomBuilder(InstanceDetails);
 }
 
diff --git a/Engine/Source/Editor/StructUtilsEditor/Private/Customizations/STypeSelector.cpp b/Engine/Source/Editor/StructUtilsEditor/Private/Customizations/STypeSelector.cpp
new file mode 100644
--- /dev/null
+++ b/Engine/Source/Editor/StructUtilsEditor/Private/Customizations/STypeSelector.cpp
@@ -0,0 +1,219 @@
+// Copyright Epic Games, Inc. All Rights Reserved.
+
+#include "STypeSelector.h"
+
+#include "PinTypeSelectorFilter.h"
+#include "Framework/Application/SlateApplication.h"
+#include "Kismet2/BlueprintEditorUtils.h"
+#include "Layout/WidgetPath.h"
+
+#include "Widgets/Images/SLayeredImage.h"
+#include "Widgets/Input/SComboButton.h"
+#include "Widgets/Layout/SWidgetSwitcher.h"
+
+/** Exact replication from SPinTypeSelector, which is not exposed. */
+class FPinTypeSelectorCustomFilterProxy final : public IPinTypeSelectorFilter
+{
+public:
+	FPinTypeSelectorCustomFilterProxy(TSharedRef<IPinTypeSelectorFilter> InFilter, FSimpleDelegate InOnFilterChanged)
+		: Filter(InFilter)
+	{
+		// Auto-register the given delegate to respond to any filter change event and refresh the filtered item list, etc.
+		OnFilterChanged_DelegateHandle = Filter->RegisterOnFilterChanged(InOnFilterChanged);
+	}
+
+	~FPinTypeSelectorCustomFilterProxy() override
+	{
+		// Auto-unregister the delegate that was previously registered at construction time.
+		Filter->UnregisterOnFilterChanged(OnFilterChanged_DelegateHandle);
+	}
+
+	virtual FDelegateHandle RegisterOnFilterChanged(const FSimpleDelegate InOnFilterChanged) override
+	{
+		return Filter->RegisterOnFilterChanged(InOnFilterChanged);
+	}
+
+	virtual void UnregisterOnFilterChanged(const FDelegateHandle InDelegateHandle) override
+	{
+		Filter->UnregisterOnFilterChanged(InDelegateHandle);
+	}
+
+	virtual TSharedPtr<SWidget> GetFilterOptionsWidget() override
+	{
+		return Filter->GetFilterOptionsWidget();
+	}
+
+	virtual bool ShouldShowPinTypeTreeItem(const FPinTypeTreeItem InItem) const override
+	{
+		return Filter->ShouldShowPinTypeTreeItem(InItem);
+	}
+
+private:
+	/** The underlying filter for which we're acting as a proxy. */
+	TSharedRef<IPinTypeSelectorFilter> Filter;
+
+	/** A handle to a delegate that gets called whenever the custom filter changes. Will be unregistered automatically when the proxy is destroyed. */
+	FDelegateHandle OnFilterChanged_DelegateHandle;
+};
+
+void STypeSelector::Construct(const FArguments& InArgs, FGetPinTypeTree GetPinTypeTreeFunc)
+{
+	// Currently only wrapping around the Compact selector type. All others should pass to the SPinTypeSelector.
+	if (InArgs._SelectorType != ESelectorType::Compact)
+	{
+		// Forward all arguments, construct, and return.
+		SPinTypeSelector::FArguments ParentArgs;
+		ParentArgs._TargetPinType = InArgs._TargetPinType;
+		ParentArgs._Schema = InArgs._Schema;
+		ParentArgs._SchemaAction = InArgs._SchemaAction;
+		ParentArgs._TypeTreeFilter = InArgs._TypeTreeFilter;
+		ParentArgs._bAllowArrays = InArgs._bAllowContainers;
+		ParentArgs._TreeViewWidth = InArgs._TreeViewWidth;
+		ParentArgs._TreeViewHeight = InArgs._TreeViewHeight;
+		ParentArgs._OnPinTypePreChanged = InArgs._OnPinTypePreChanged;
+		ParentArgs._OnPinTypeChanged = InArgs._OnPinTypeChanged;
+		ParentArgs._SelectorType = InArgs._SelectorType;
+		ParentArgs._ReadOnly = InArgs._ReadOnly;
+		ParentArgs._CustomFilters = InArgs._CustomFilters;
+		ParentArgs._OnPinTypeChanged = InArgs._OnPinTypeChanged;
+		SPinTypeSelector::Construct(ParentArgs, GetPinTypeTreeFunc);
+		return;
+	}
+
+	SearchText = FText::GetEmpty();
+
+	ReadOnly = InArgs._ReadOnly;
+
+	OnTypeChanged = InArgs._OnPinTypeChanged;
+	OnTypePreChanged = InArgs._OnPinTypePreChanged;
+
+	check(GetPinTypeTreeFunc.IsBound());
+	GetPinTypeTree = GetPinTypeTreeFunc;
+
+	Schema = InArgs._Schema;
+	SchemaAction = InArgs._SchemaAction;
+	TypeTreeFilter = InArgs._TypeTreeFilter;
+	TreeViewWidth = InArgs._TreeViewWidth;
+	TreeViewHeight = InArgs._TreeViewHeight;
+
+	TargetPinType = InArgs._TargetPinType;
+	SelectorType = InArgs._SelectorType;
+
+	NumFilteredPinTypeItems = 0;
+	NumValidPinTypeItems = 0;
+
+	bIsRightMousePressed = false;
+
+	if (InArgs._CustomFilters.Num() > 0)
+	{
+		for (const TSharedPtr<IPinTypeSelectorFilter>& Filter : InArgs._CustomFilters)
+		{
+			CustomFilters.Add(MakeShared<FPinTypeSelectorCustomFilterProxy>(Filter.ToSharedRef(), FSimpleDelegate::CreateSP(this, &STypeSelector::OnCustomFilterChanged)));
+		}
+	}
+	else if (UClass* PinTypeSelectorFilterClass = GetDefault<UPinTypeSelectorFilter>()->FilterClass.LoadSynchronous())
+	{
+		const TSharedPtr<IPinTypeSelectorFilter> SelectorFilter = GetDefault<UPinTypeSelectorFilter>(PinTypeSelectorFilterClass)->GetPinTypeSelectorFilter();
+		CustomFilters.Add(MakeShared<FPinTypeSelectorCustomFilterProxy>(SelectorFilter.ToSharedRef(), FSimpleDelegate::CreateSP(this, &STypeSelector::OnCustomFilterChanged)));
+	}
+
+	const TSharedPtr<SWidget> ReadOnlyWidget = SNew(SHorizontalBox)
+		.Clipping(EWidgetClipping::OnDemand)
+		+ SHorizontalBox::Slot()
+		.VAlign(VAlign_Center)
+		.HAlign(HAlign_Left)
+		.Padding(FMargin(2.0f, 3.0f, 2.0f, 3.0f))
+		.AutoWidth()
+		[
+			// Read-only version does not display container or secondary type separately, so we need to jam it all in the one image
+			SNew(SLayeredImage, TAttribute<const FSlateBrush*>(this, &STypeSelector::GetSecondaryTypeIconImage), TAttribute<FSlateColor>(this, &STypeSelector::GetSecondaryTypeIconColor))
+			.Image(this, &STypeSelector::GetTypeIconImage)
+			.ColorAndOpacity(this, &STypeSelector::GetTypeIconColor)
+		]
+		+ SHorizontalBox::Slot()
+		.Padding(2.0f, 2.0f)
+		.VAlign(VAlign_Center)
+		.HAlign(HAlign_Left)
+		.AutoWidth()
+		[
+			SNew(STextBlock)
+			.Text(this, &STypeSelector::GetTypeDescription, false)
+			.Font(InArgs._Font)
+			.ColorAndOpacity(FSlateColor::UseForeground())
+		];
+
+	SAssignNew(TypeComboButton, SComboButton)
+	.OnGetMenuContent(this, &STypeSelector::GetMenuContent, false)
+	.ContentPadding(0.0f)
+	.ToolTipText(this, &STypeSelector::GetToolTipForSelector)
+	.HasDownArrow(true)
+	.ButtonStyle(FAppStyle::Get(), "SimpleButton")
+	.ButtonContent()
+	[
+		SNew(SLayeredImage,
+			TAttribute<const FSlateBrush*>(this, &STypeSelector::GetSecondaryTypeIconImage),
+			TAttribute<FSlateColor>(this, &STypeSelector::GetSecondaryTypeIconColor))
+		.Image(this, &STypeSelector::GetTypeIconImage)
+		.ColorAndOpacity(this, &STypeSelector::GetTypeIconColor)
+	];
+
+	this->ChildSlot
+	[
+		SNew(SWidgetSwitcher)
+		.WidgetIndex_Lambda([this]() { return ReadOnly.Get() ? 1 : 0; })
+		+ SWidgetSwitcher::Slot() // editable version
+		.Padding(SelectorType == ESelectorType::Partial ? FMargin(-6.0f, 0.0f, 0.0f, 0.0f) : FMargin(0))
+		[
+			TypeComboButton.ToSharedRef()
+		]
+		+ SWidgetSwitcher::Slot() // read-only version
+		[
+			ReadOnlyWidget.ToSharedRef()
+		]
+	];
+}
+
+FReply STypeSelector::OnMouseButtonDown(const FGeometry& MyGeometry, const FPointerEvent& MouseEvent)
+{
+	if (MouseEvent.GetEffectingButton() == EKeys::RightMouseButton)
+	{
+		bIsRightMousePressed = true;
+	}
+
+	return SPinTypeSelector::OnMouseButtonDown(MyGeometry, MouseEvent);
+}
+
+FReply STypeSelector::OnMouseButtonUp(const FGeometry& MyGeometry, const FPointerEvent& MouseEvent)
+{
+	if (MouseEvent.GetEffectingButton() == EKeys::RightMouseButton && bIsRightMousePressed)
+	{
+		// Push the other menu from the secondary ComboButton for the type.
+		FSlateApplication::Get().PushMenu(TypeComboButton.ToSharedRef(),
+			FWidgetPath(),
+			GetPinContainerTypeMenuContent(),
+			MyGeometry.AbsolutePosition,
+			FPopupTransitionEffect(FPopupTransitionEffect::SubMenu));
+
+		bIsRightMousePressed = false;
+		return FReply::Handled();
+	}
+	else
+	{
+		bIsRightMousePressed = false;
+		return SPinTypeSelector::OnMouseButtonUp(MyGeometry, MouseEvent);
+	}
+}
+
+FText STypeSelector::GetToolTipForSelector() const
+{
+	FText EditText;
+	if (IsEnabled())
+	{
+		EditText = NSLOCTEXT("STypeSelector", "PinTypeSelectorTooltip", "Left click to select the variable's type. Right click to select a container type.\n");
+	}
+	else
+	{
+		EditText = NSLOCTEXT("STypeSelector", "DisabledPinTypeSelectorTooltip", "Cannot edit variable type when they are inherited from parent.\n");
+	}
+	return EditText;
+}
diff --git a/Engine/Source/Editor/StructUtilsEditor/Private/Customizations/STypeSelector.h b/Engine/Source/Editor/StructUtilsEditor/Private/Customizations/STypeSelector.h
new file mode 100644
--- /dev/null
+++ b/Engine/Source/Editor/StructUtilsEditor/Private/Customizations/STypeSelector.h
@@ -0,0 +1,57 @@
+// Copyright Epic Games, Inc. All Rights Reserved.
+
+#pragma once
+
+#include "SPinTypeSelector.h"
+
+/**
+ * This widget is a small wrapper around SPinTypeSelector to allow for a right click context menu on the selector pill
+ * for the container type so that it can all fit into one compact combo button.
+ * 
+ * Eventually (TODO), it should be replaced by a base combo button drop down that supports right clicks as
+ * well as an update to the original SPinTypeSelector to support this for Compact Selector Type (rather than toggling).
+ * However, this will require updates to core widgets, which are currently hardcoded for left click only.
+ * 
+ * Ex. 'SDoubleComboButton' (SComboButton) which inherits from 'SDoubleMenuAnchor' (SMenuAnchor) or similar with two
+ * menu anchors. */
+UE_EXPERIMENTAL(5.6, "This widget is an experimental prototype for StructUtils use only.")
+class STRUCTUTILSEDITOR_API STypeSelector final : public SPinTypeSelector
+{
+public:
+	// Replicated arguments of SPinTypeSelector
+	SLATE_BEGIN_ARGS(STypeSelector)
+			: _Schema(nullptr)
+			, _SchemaAction(nullptr)
+			, _TypeTreeFilter(ETypeTreeFilter::None)
+			, _bAllowContainers(true)
+			, _TreeViewWidth(300.f)
+			, _TreeViewHeight(350.f)
+			, _Font(FAppStyle::GetFontStyle(TEXT("NormalFont")))
+			, _SelectorType(ESelectorType::Full)
+			, _ReadOnly(false) {}
+
+		SLATE_ATTRIBUTE(FEdGraphPinType, TargetPinType)
+		SLATE_ARGUMENT(const UEdGraphSchema*, Schema)
+		SLATE_ARGUMENT(TWeakPtr<const FEdGraphSchemaAction>, SchemaAction)
+		SLATE_ARGUMENT(ETypeTreeFilter, TypeTreeFilter)
+		SLATE_ARGUMENT(bool, bAllowContainers)
+		SLATE_ATTRIBUTE(FOptionalSize, TreeViewWidth)
+		SLATE_ATTRIBUTE(FOptionalSize, TreeViewHeight)
+		SLATE_EVENT(FOnPinTypeChanged, OnPinTypePreChanged)
+		SLATE_EVENT(FOnPinTypeChanged, OnPinTypeChanged)
+		SLATE_ATTRIBUTE(FSlateFontInfo, Font)
+		SLATE_ARGUMENT(ESelectorType, SelectorType)
+		SLATE_ATTRIBUTE(bool, ReadOnly)
+		SLATE_ARGUMENT(TArray<TSharedPtr<class IPinTypeSelectorFilter>>, CustomFilters)
+	SLATE_END_ARGS()
+
+	void Construct(const FArguments& InArgs, FGetPinTypeTree GetPinTypeTreeFunc);
+
+	// SWidget interface
+	virtual FReply OnMouseButtonDown(const FGeometry& MyGeometry, const FPointerEvent& MouseEvent) override;
+	virtual FReply OnMouseButtonUp(const FGeometry& MyGeometry, const FPointerEvent& MouseEvent) override;
+	// End of SWidget interface
+
+private:
+	FText GetToolTipForSelector() const;
+};
diff --git a/Engine/Source/Editor/StructUtilsEditor/Private/StructUtilsMetadata.cpp b/Engine/Source/Editor/StructUtilsEditor/Private/StructUtilsMetadata.cpp
--- a/Engine/Source/Editor/StructUtilsEditor/Private/StructUtilsMetadata.cpp
+++ b/Engine/Source/Editor/StructUtilsEditor/Private/StructUtilsMetadata.cpp
@@ -4,9 +4,14 @@
 
 namespace UE::StructUtils::Metadata
 {
+	const FName HideInDetailPanelsName = TEXT("HideInDetailPanel");
+	const FName ShowOnlyInnerPropertiesName = TEXT("ShowOnlyInnerProperties");
+
 	const FName FixedLayoutName = TEXT("FixedLayout");
 	const FName DefaultTypeName = TEXT("DefaultType");
+	const FName AllowContainersName = TEXT("AllowContainers");
 	const FName AllowArraysName = TEXT("AllowArrays");
 	const FName IsPinTypeAcceptedName = TEXT("IsPinTypeAccepted");
 	const FName CanRemovePropertyName = TEXT("CanRemoveProperty");
+	const FName ChildRowFeaturesName = TEXT("ChildRowFeatures");
 }
\ No newline at end of file
diff --git a/Engine/Source/Editor/StructUtilsEditor/Public/PropertyBagDetails.h b/Engine/Source/Editor/StructUtilsEditor/Public/PropertyBagDetails.h
--- a/Engine/Source/Editor/StructUtilsEditor/Public/PropertyBagDetails.h
+++ b/Engine/Source/Editor/StructUtilsEditor/Public/PropertyBagDetails.h
@@ -8,8 +8,9 @@
 #include "PropertyBagDetails.generated.h"
 
 class FReply;
-enum class EPinContainerType : uint8;
+enum class ESelectorType : uint8;
 struct FEdGraphSchemaAction;
+struct FMenuEntryParams;
 
 class IPropertyHandle;
 class IPropertyUtilities;
@@ -17,10 +18,53 @@ class IDetailPropertyRow;
 class SInlineEditableTextBlock;
 class SWidget;
 
+/**
+ * The property bag details child rows can be completely customized by assigning a combination of these values
+ * to their FPropertyBag 'ChildRowFeatures' metadata via the meta specifier. TODO: It isn't currently possible
+ * to specify bitwise operations in the metadata string, but this will be added shortly. More configurations can be
+ * added in the meantime to satisfy specific needs.
+ */
+UENUM(meta = (Bitflags))
+enum class EPropertyBagChildRowFeatures : uint64
+{
+	/** General Options */
+	Fixed                  = 0,      // Fixed layout. No features enabled.
+	Renaming               = 1ULL << 0, // The property name is editable from the details view.
+	Deletion               = 1ULL << 1, // The property is able to be deleted.
+	DragAndDrop            = 1ULL << 2, // Drag and dropping properties is enabled.
+	CompactTypeSelector    = 1ULL << 3, // A compact type selector widget is enabled to the left of the name.
+	AccessSpecifierButton  = 1ULL << 4, // The property metadata can be set for public/private access.
+	DropDownMenuButton     = 1ULL << 5, // Drop-down menu (down arrow button) is enabled.
+	AllGeneralOptions = Renaming | Deletion | DragAndDrop | CompactTypeSelector | AccessSpecifierButton | DropDownMenuButton,
+
+	// Insert new drop general features above this line.
+
+	/** Menu Options (for drop-down menu or other future menus (right click) */
+	Menu_TypeSelector = 1ULL << 17, // Adds the type selector pill widget to the drop-down.
+	Menu_Rename       = 1ULL << 18, // Renaming the property from the drop-down menu. Requires property renaming enabled.
+	Menu_Delete       = 1ULL << 19, // Deleting the property from the drop-down menu.
+	AllMenuOptions    = Menu_TypeSelector | Menu_Rename | Menu_Delete,
+
+	// Insert new menu features above this line.
+	Deprecated = 1ULL << 63, // To allow for deprecating older features.
+
+	// Below are configurations for convenience.
+	ReadOnly = Fixed,
+	Core     = Renaming | DropDownMenuButton | AllMenuOptions,
+	Extended = Renaming | Deletion | DragAndDrop | CompactTypeSelector | AccessSpecifierButton | DropDownMenuButton | Menu_Rename | Menu_Delete,
+	All      = AllGeneralOptions | AllMenuOptions,
+	// The default version includes deprecated UI features to support previous behavior.
+	Default = ReadOnly | Deprecated
+};
+
+ENUM_CLASS_FLAGS(EPropertyBagChildRowFeatures);
+
 namespace UE::StructUtils
 {
 	STRUCTUTILSEDITOR_API void SetPropertyDescFromPin(FPropertyBagPropertyDesc& Desc, const FEdGraphPinType& PinType);
 	STRUCTUTILSEDITOR_API FEdGraphPinType GetPropertyDescAsPin(const FPropertyBagPropertyDesc& Desc);
+	/** Creates type selection pill widget. */
+	STRUCTUTILSEDITOR_API TSharedRef<SWidget> CreateTypeSelectionWidget(TSharedPtr<IPropertyHandle> ChildPropertyHandle, const TSharedPtr<IPropertyHandle>& InBagStructProperty, const TSharedPtr<IPropertyUtilities>& InPropUtils, ESelectorType SelectorType, bool bAllowContainers = true);
 }
 
 /**
@@ -34,7 +78,7 @@ public:
 
 	/** Creates add property widget. */
 	static TSharedPtr<SWidget> MakeAddPropertyWidget(TSharedPtr<IPropertyHandle> InStructProperty, TSharedPtr<IPropertyUtilities> InPropUtils, EPropertyBagPropertyType DefaultType = EPropertyBagPropertyType::Bool, FSlateColor IconColor = FSlateColor::UseForeground());
-	
+
 protected:
 	/** IPropertyTypeCustomization interface */
 	virtual void CustomizeHeader(TSharedRef<IPropertyHandle> StructPropertyHandle, FDetailWidgetRow& HeaderRow, IPropertyTypeCustomizationUtils& StructCustomizationUtils) override;
@@ -45,8 +89,10 @@ private:
 	TSharedPtr<IPropertyHandle> StructProperty;
 	TSharedPtr<IPropertyUtilities> PropUtils;
 	EPropertyBagPropertyType DefaultType = EPropertyBagPropertyType::Bool;
+	EPropertyBagChildRowFeatures ChildRowFeatures = EPropertyBagChildRowFeatures::Default;
+
 	bool bFixedLayout = false;
-	bool bAllowArrays = true;
+	bool bAllowContainers = true;
 };
 
 /** 
@@ -60,8 +106,16 @@ private:
 class STRUCTUTILSEDITOR_API FPropertyBagInstanceDataDetails : public FInstancedStructDataDetails
 {
 public:
+	struct FConstructParams
+	{
+		TSharedPtr<IPropertyHandle> BagStructProperty = nullptr;
+		TSharedPtr<IPropertyUtilities> PropUtils = nullptr;
+		bool bAllowContainers = true;
+		EPropertyBagChildRowFeatures ChildRowFeatures = EPropertyBagChildRowFeatures::Default;
+	};
 
-	FPropertyBagInstanceDataDetails(TSharedPtr<IPropertyHandle> InStructProperty, const TSharedPtr<IPropertyUtilities>& InPropUtils, const bool bInFixedLayout, const bool bInAllowArrays = true);
+	explicit FPropertyBagInstanceDataDetails(const FConstructParams& ConstructParams);
+	FPropertyBagInstanceDataDetails(TSharedPtr<IPropertyHandle> InStructProperty, const TSharedPtr<IPropertyUtilities>& InPropUtils, bool bInFixedLayout, bool bInAllowContainers = true);
 
 	virtual void OnChildRowAdded(IDetailPropertyRow& ChildRow) override;
 	
@@ -124,12 +178,12 @@ public:
 	 */
 	virtual void EnumeratePropertyBags(TSharedPtr<IPropertyHandle> PropertyBagHandle, const EnumeratePropertyBagFuncRef& Func) const
 	{
-		checkf(false, TEXT("EnumeratePropertyBags() is expecgted to be implemented when HasPropertyOverrides() returns true."));
+		checkf(false, TEXT("EnumeratePropertyBags() is expected to be implemented when HasPropertyOverrides() returns true."));
 	}
 
 	/** @return true if property of specified child property is overridden. */
 	virtual EPropertyOverrideState IsPropertyOverridden(TSharedPtr<IPropertyHandle> ChildPropertyHandle) const;
-	
+
 	/** Called to set the override state of specified child property. */
 	virtual void SetPropertyOverride(TSharedPtr<IPropertyHandle> ChildPropertyHandle, const bool bIsOverridden);
 
@@ -138,13 +192,20 @@ public:
 
 	/** Called to reset the child property to default value. */
 	virtual void ResetToDefault(TSharedPtr<IPropertyHandle> ChildPropertyHandle);
-	
+
 protected:
+	UE_DEPRECATED(5.6, "Replaced by 'CreatePropertyDetailsWidget' to allow for more customizable drop-down menu location and content.")
 	TSharedRef<SWidget> OnPropertyNameContent(TSharedPtr<IPropertyHandle> ChildPropertyHandle, TSharedPtr<SInlineEditableTextBlock> InlineWidget) const;
 
 	TSharedPtr<IPropertyHandle> BagStructProperty;
 	TSharedPtr<IPropertyUtilities> PropUtils;
+
+	bool bAllowContainers = true;
+	EPropertyBagChildRowFeatures ChildRowFeatures = EPropertyBagChildRowFeatures::Default;
+
+	UE_DEPRECATED(5.6, "Use 'ChildRowFeatures' instead.")
 	bool bFixedLayout = false;
+	UE_DEPRECATED(5.6, "Use 'bAllowContainers' instead.")
 	bool bAllowArrays = true;
 };
 
diff --git a/Engine/Source/Editor/StructUtilsEditor/Public/StructUtilsMetadata.h b/Engine/Source/Editor/StructUtilsEditor/Public/StructUtilsMetadata.h
--- a/Engine/Source/Editor/StructUtilsEditor/Public/StructUtilsMetadata.h
+++ b/Engine/Source/Editor/StructUtilsEditor/Public/StructUtilsMetadata.h
@@ -9,6 +9,10 @@ namespace UE::StructUtils::Metadata
 {
 	// Metadata usable in UPROPERTY for customizing the behavior when displaying the property in a property panel or graph node
 
+	/// Common UPROPERTY metadata names.
+	extern STRUCTUTILSEDITOR_API const FName HideInDetailPanelsName;
+	extern STRUCTUTILSEDITOR_API const FName ShowOnlyInnerPropertiesName;
+
 	/// FInstancedPropertyBag
 
 	/// [PropertyMetadata] FixedLayout: Indicates that the instanced property bag has a fixed layout and it is not possible to add/remove properties.
@@ -17,7 +21,10 @@ namespace UE::StructUtils::Metadata
 	/// [PropertyMetadata] DefaultType: Default property type when adding a new Property. Should be taken from enum EPropertyBagPropertyType
 	extern STRUCTUTILSEDITOR_API const FName DefaultTypeName;
 
-	/// [PropertyMetadata] AllowArrays: By default it is always allowed to create arrays for properties. Use AllowArrays=false to disable Array support.
+	/// [PropertyMetadata] AllowContainers: By default it is always allowed to create containers (Array, Set) for properties. Use AllowContainers=false to disable container support.
+	extern STRUCTUTILSEDITOR_API const FName AllowContainersName;
+
+	UE_DEPRECATED(5.6, "Remains for backwards compatibility. Please see 'AllowContainersName'.")
 	extern STRUCTUTILSEDITOR_API const FName AllowArraysName;
 
 	/// [PropertyMetadata] IsPinTypeAccepted: Name of a UFunction with signature bool(FEdGraphPinType). Returns false if the type should be discarded.
@@ -25,4 +32,7 @@ namespace UE::StructUtils::Metadata
 
 	/// [PropertyMetadata] CanRemoveProperty: Name of a UFunction with signature bool(FGuid, FName). ID and name of the property that will be removed. Returns false if the property should not be removed.
 	extern STRUCTUTILSEDITOR_API const FName CanRemovePropertyName;
-}
\ No newline at end of file
+
+	/// [PropertyMetadata] ChildRowFeatures: A list of UI features available to the PropertyBag properties when displayed in the details view child rows.
+	extern STRUCTUTILSEDITOR_API const FName ChildRowFeaturesName;
+}
diff --git a/Engine/Source/Runtime/CoreUObject/Public/StructUtils/PropertyBag.h b/Engine/Source/Runtime/CoreUObject/Public/StructUtils/PropertyBag.h
--- a/Engine/Source/Runtime/CoreUObject/Public/StructUtils/PropertyBag.h
+++ b/Engine/Source/Runtime/CoreUObject/Public/StructUtils/PropertyBag.h
@@ -303,7 +303,7 @@ struct FPropertyBagPropertyDesc
 	 * Serializes a PropertyBagPropertyDesc from or into an archive.
 	 *
 	 * @param Ar The archive to serialize from or into.
-	 * @param FPropertyBagPropertyDesc The PropertyBagPropertyDesc to serialize.
+	 * @param Desc The PropertyBagPropertyDesc to serialize.
 	 */
 	 UE_API friend FArchive& operator<<(FArchive& Ar, FPropertyBagPropertyDesc& Desc);
 
@@ -885,14 +885,14 @@ struct FInstancedPropertyBag
 	 * @param Name The name to sanitize.
 	 * @return Sanitized property name based on the input string.
 	 */
-	static FName SanitizePropertyName(const FString& Name);
+	UE_API static FName SanitizePropertyName(const FString& Name);
 
 	/**
 	 * Returns a sanitized version of the provided name without invalid characters.
 	 * @param Name The name to sanitize.
 	 * @return Sanitized property name based on the input name.
 	 */
-	static FName SanitizePropertyName(FName Name);
+	UE_API static FName SanitizePropertyName(FName Name);
 
 protected:
 	UE_API const void* GetValueAddress(const FPropertyBagPropertyDesc* Desc) const;
