[StructUtils] Fix deprecation warning for bFixedLayout

#rnx

diff --git a/Engine/Source/Editor/StructUtilsEditor/Private/Customizations/PropertyBagDetails.cpp b/Engine/Source/Editor/StructUtilsEditor/Private/Customizations/PropertyBagDetails.cpp
--- a/Engine/Source/Editor/StructUtilsEditor/Private/Customizations/PropertyBagDetails.cpp
+++ b/Engine/Source/Editor/StructUtilsEditor/Private/Customizations/PropertyBagDetails.cpp
@@ -1677,7 +1677,9 @@ void FPropertyBagDetails::CustomizeHeader(TSharedRef<IPropertyHandle> StructProp
 
 	if (const FProperty* MetaDataProperty = StructProperty->GetMetaDataProperty())
 	{
+		PRAGMA_DISABLE_DEPRECATION_WARNINGS
 		bFixedLayout = MetaDataProperty->HasMetaData(UE::StructUtils::Metadata::FixedLayoutName);
+		PRAGMA_ENABLE_DEPRECATION_WARNINGS
 
 		bAllowContainers = MetaDataProperty->HasMetaData(UE::StructUtils::Metadata::AllowContainersName)
 			? MetaDataProperty->GetBoolMetaData(UE::StructUtils::Metadata::AllowContainersName)
@@ -1717,7 +1719,9 @@ void FPropertyBagDetails::CustomizeHeader(TSharedRef<IPropertyHandle> StructProp
 	}
 
 	TSharedPtr<SWidget> ValueWidget = SNullWidget::NullWidget;
-	if (!bFixedLayout)
+	PRAGMA_DISABLE_DEPRECATION_WARNINGS
+	if (!bFixedLayout && ChildRowFeatures != EPropertyBagChildRowFeatures::Fixed)
+	PRAGMA_ENABLE_DEPRECATION_WARNINGS
 	{
 		ValueWidget = MakeAddPropertyWidget(StructProperty, PropUtils, DefaultType);
 	}
diff --git a/Engine/Source/Editor/StructUtilsEditor/Public/PropertyBagDetails.h b/Engine/Source/Editor/StructUtilsEditor/Public/PropertyBagDetails.h
--- a/Engine/Source/Editor/StructUtilsEditor/Public/PropertyBagDetails.h
+++ b/Engine/Source/Editor/StructUtilsEditor/Public/PropertyBagDetails.h
@@ -91,6 +91,7 @@ private:
 	EPropertyBagPropertyType DefaultType = EPropertyBagPropertyType::Bool;
 	EPropertyBagChildRowFeatures ChildRowFeatures = EPropertyBagChildRowFeatures::Default;
 
+	UE_DEPRECATED(5.6, "Use 'ChildRowFeatures' instead.")
 	bool bFixedLayout = false;
 	bool bAllowContainers = true;
 };
