[StructUtils] Fix more deprecation issues related to implicit copy contruction

#rnx

diff --git a/Engine/Source/Editor/StructUtilsEditor/Private/Customizations/PropertyBagDetails.cpp b/Engine/Source/Editor/StructUtilsEditor/Private/Customizations/PropertyBagDetails.cpp
--- a/Engine/Source/Editor/StructUtilsEditor/Private/Customizations/PropertyBagDetails.cpp
+++ b/Engine/Source/Editor/StructUtilsEditor/Private/Customizations/PropertyBagDetails.cpp
@@ -981,6 +981,10 @@ FPropertyBagInstanceDataDetails::FPropertyBagInstanceDataDetails(const FProperty
 	, PropUtils(ConstructParams.PropUtils)
 	, bAllowContainers(ConstructParams.bAllowContainers)
 	, ChildRowFeatures(ConstructParams.ChildRowFeatures)
+PRAGMA_DISABLE_DEPRECATION_WARNINGS
+	, bFixedLayout(false)
+	, bAllowArrays(true)
+PRAGMA_ENABLE_DEPRECATION_WARNINGS
 {
 	ensure(UE::StructUtils::Private::IsScriptStruct<FInstancedPropertyBag>(BagStructProperty));
 	ensure(PropUtils != nullptr);
@@ -994,6 +998,7 @@ FPropertyBagInstanceDataDetails::FPropertyBagInstanceDataDetails(TSharedPtr<IPro
 	, bAllowContainers(bInAllowContainers)
 	, ChildRowFeatures(bInFixedLayout ? EPropertyBagChildRowFeatures::Fixed : EPropertyBagChildRowFeatures::Default)
 PRAGMA_DISABLE_DEPRECATION_WARNINGS
+	, bFixedLayout(bInFixedLayout) // For backwards compatibility
 	, bAllowArrays(bInAllowContainers) // For backwards compatibility
 PRAGMA_ENABLE_DEPRECATION_WARNINGS
 {
@@ -1808,4 +1813,4 @@ bool UPropertyBagSchema::SupportsPinTypeContainer(TWeakPtr<const FEdGraphSchemaA
 	return ContainerType == EPinContainerType::None || ContainerType == EPinContainerType::Array || ContainerType == EPinContainerType::Set;
 }
 
-#undef LOCTEXT_NAMESPACE
+#undef LOCTEXT_NAMESPA
\ No newline at end of file
