[StructUtils] Another deprecation warning from an implicit construction.

#rnx

diff --git a/Engine/Source/Editor/StructUtilsEditor/Private/Customizations/PropertyBagDetails.cpp b/Engine/Source/Editor/StructUtilsEditor/Private/Customizations/PropertyBagDetails.cpp
--- a/Engine/Source/Editor/StructUtilsEditor/Private/Customizations/PropertyBagDetails.cpp
+++ b/Engine/Source/Editor/StructUtilsEditor/Private/Customizations/PropertyBagDetails.cpp
@@ -1668,6 +1668,11 @@ TSharedRef<SWidget> FPropertyBagInstanceDataDetails::OnPropertyNameContent(TShar
 
 ////////////////////////////////////
 
+FPropertyBagDetails::FPropertyBagDetails()
+PRAGMA_DISABLE_DEPRECATION_WARNINGS
+	: bFixedLayout(false) {} // Implicit construction was causing deprecation warnings
+PRAGMA_ENABLE_DEPRECATION_WARNINGS
+
 TSharedRef<IPropertyTypeCustomization> FPropertyBagDetails::MakeInstance()
 {
 	return MakeShared<FPropertyBagDetails>();
@@ -1813,4 +1818,4 @@ bool UPropertyBagSchema::SupportsPinTypeContainer(TWeakPtr<const FEdGraphSchemaA
 	return ContainerType == EPinContainerType::None || ContainerType == EPinContainerType::Array || ContainerType == EPinContainerType::Set;
 }
 
-#undef LOCTEXT_NAMESPA
\ No newline at end of file
+#undef LOCTEXT_NAMESPACE
diff --git a/Engine/Source/Editor/StructUtilsEditor/Public/PropertyBagDetails.h b/Engine/Source/Editor/StructUtilsEditor/Public/PropertyBagDetails.h
--- a/Engine/Source/Editor/StructUtilsEditor/Public/PropertyBagDetails.h
+++ b/Engine/Source/Editor/StructUtilsEditor/Public/PropertyBagDetails.h
@@ -73,6 +73,8 @@ namespace UE::StructUtils
 class STRUCTUTILSEDITOR_API FPropertyBagDetails : public IPropertyTypeCustomization
 {
 public:
+	FPropertyBagDetails();
+
 	/** Makes a new instance of this detail layout class for a specific detail view requesting it */
 	static TSharedRef<IPropertyTypeCustomization> MakeInstance();
 
