[StructUtils] Backout Deprecation of FPropertyBagDetails::bFixedLayout

#rnx

diff --git a/Engine/Source/Editor/StructUtilsEditor/Private/Customizations/PropertyBagDetails.cpp b/Engine/Source/Editor/StructUtilsEditor/Private/Customizations/PropertyBagDetails.cpp
--- a/Engine/Source/Editor/StructUtilsEditor/Private/Customizations/PropertyBagDetails.cpp
+++ b/Engine/Source/Editor/StructUtilsEditor/Private/Customizations/PropertyBagDetails.cpp
@@ -1668,11 +1668,6 @@ TSharedRef<SWidget> FPropertyBagInstanceDataDetails::OnPropertyNameContent(TShar
 
 ////////////////////////////////////
 
-FPropertyBagDetails::FPropertyBagDetails()
-PRAGMA_DISABLE_DEPRECATION_WARNINGS
-	: bFixedLayout(false) {} // Implicit construction was causing deprecation warnings
-PRAGMA_ENABLE_DEPRECATION_WARNINGS
-
 TSharedRef<IPropertyTypeCustomization> FPropertyBagDetails::MakeInstance()
 {
 	return MakeShared<FPropertyBagDetails>();
diff --git a/Engine/Source/Editor/StructUtilsEditor/Public/PropertyBagDetails.h b/Engine/Source/Editor/StructUtilsEditor/Public/PropertyBagDetails.h
--- a/Engine/Source/Editor/StructUtilsEditor/Public/PropertyBagDetails.h
+++ b/Engine/Source/Editor/StructUtilsEditor/Public/PropertyBagDetails.h
@@ -73,8 +73,6 @@ namespace UE::StructUtils
 class STRUCTUTILSEDITOR_API FPropertyBagDetails : public IPropertyTypeCustomization
 {
 public:
-	FPropertyBagDetails();
-
 	/** Makes a new instance of this detail layout class for a specific detail view requesting it */
 	static TSharedRef<IPropertyTypeCustomization> MakeInstance();
 
@@ -93,7 +91,6 @@ private:
 	EPropertyBagPropertyType DefaultType = EPropertyBagPropertyType::Bool;
 	EPropertyBagChildRowFeatures ChildRowFeatures = EPropertyBagChildRowFeatures::Default;
 
-	UE_DEPRECATED(5.6, "Use 'ChildRowFeatures' instead.")
 	bool bFixedLayout = false;
 	bool bAllowContainers = true;
 };
