[StructUtils] Add explicit move/copy constructors for PropertyBagDetails.h to address deprecation warnings

#rnx

diff --git a/Engine/Source/Editor/StructUtilsEditor/Public/PropertyBagDetails.h b/Engine/Source/Editor/StructUtilsEditor/Public/PropertyBagDetails.h
--- a/Engine/Source/Editor/StructUtilsEditor/Public/PropertyBagDetails.h
+++ b/Engine/Source/Editor/StructUtilsEditor/Public/PropertyBagDetails.h
@@ -73,6 +73,16 @@ namespace UE::StructUtils
 class STRUCTUTILSEDITOR_API FPropertyBagDetails : public IPropertyTypeCustomization
 {
 public:
+	// Required for deprecated variable 'bFixedLayout'. TODO: Remove when removing 5.6 deprecations.
+	PRAGMA_DISABLE_DEPRECATION_WARNINGS
+	FPropertyBagDetails() = default;
+	~FPropertyBagDetails() = default;
+	FPropertyBagDetails(const FPropertyBagDetails& Other) = delete;
+	FPropertyBagDetails(FPropertyBagDetails&& Other) noexcept = delete;
+	FPropertyBagDetails& operator=(const FPropertyBagDetails& Other) = delete;
+	FPropertyBagDetails& operator=(FPropertyBagDetails&& Other) noexcept = delete;
+	PRAGMA_ENABLE_DEPRECATION_WARNINGS
+
 	/** Makes a new instance of this detail layout class for a specific detail view requesting it */
 	static TSharedRef<IPropertyTypeCustomization> MakeInstance();
 
@@ -91,6 +101,7 @@ private:
 	EPropertyBagPropertyType DefaultType = EPropertyBagPropertyType::Bool;
 	EPropertyBagChildRowFeatures ChildRowFeatures = EPropertyBagChildRowFeatures::Default;
 
+	UE_DEPRECATED(5.6, "Use 'ChildRowFeatures' instead.")
 	bool bFixedLayout = false;
 	bool bAllowContainers = true;
 };
@@ -117,6 +128,15 @@ public:
 	explicit FPropertyBagInstanceDataDetails(const FConstructParams& ConstructParams);
 	FPropertyBagInstanceDataDetails(TSharedPtr<IPropertyHandle> InStructProperty, const TSharedPtr<IPropertyUtilities>& InPropUtils, bool bInFixedLayout, bool bInAllowContainers = true);
 
+	// Required for deprecated variable 'bFixedLayout' and 'bAllowArrays'. TODO: Remove when removing 5.6 deprecations.
+	PRAGMA_DISABLE_DEPRECATION_WARNINGS
+	~FPropertyBagInstanceDataDetails() = default;
+	FPropertyBagInstanceDataDetails(const FPropertyBagInstanceDataDetails& Other) = delete;
+	FPropertyBagInstanceDataDetails(FPropertyBagInstanceDataDetails&& Other) noexcept = delete;
+	FPropertyBagInstanceDataDetails& operator=(const FPropertyBagInstanceDataDetails& Other) = delete;
+	FPropertyBagInstanceDataDetails& operator=(FPropertyBagInstanceDataDetails&& Other) noexcept = delete;
+	PRAGMA_ENABLE_DEPRECATION_WARNINGS
+
 	virtual void OnChildRowAdded(IDetailPropertyRow& ChildRow) override;
 	
 	/** Enum describing if a property is overridden, or undetermined (e.g. multiselection) */
