[StructUtils] Fixed a bug where ChildRowFeatures = "Fixed" was being bypassed incorrectly.

#rnx

diff --git a/Engine/Source/Editor/StructUtilsEditor/Private/Customizations/PropertyBagDetails.cpp b/Engine/Source/Editor/StructUtilsEditor/Private/Customizations/PropertyBagDetails.cpp
--- a/Engine/Source/Editor/StructUtilsEditor/Private/Customizations/PropertyBagDetails.cpp
+++ b/Engine/Source/Editor/StructUtilsEditor/Private/Customizations/PropertyBagDetails.cpp
@@ -1690,12 +1690,6 @@ void FPropertyBagDetails::CustomizeHeader(TSharedRef<IPropertyHandle> StructProp
 			? MetaDataProperty->GetBoolMetaData(UE::StructUtils::Metadata::AllowContainersName)
 			: true;
 
-		// Don't show the header if ShowOnlyInnerProperties is set
-		if (MetaDataProperty->HasMetaData(UE::StructUtils::Metadata::ShowOnlyInnerPropertiesName))
-		{
-			return;
-		}
-
 		if (MetaDataProperty->HasMetaData(UE::StructUtils::Metadata::DefaultTypeName))
 		{
 			if (UEnum* Enum = StaticEnum<EPropertyBagPropertyType>())
@@ -1721,6 +1715,12 @@ void FPropertyBagDetails::CustomizeHeader(TSharedRef<IPropertyHandle> StructProp
 				}
 			}
 		}
+
+		// Don't show the header if ShowOnlyInnerProperties is set
+		if (MetaDataProperty->HasMetaData(UE::StructUtils::Metadata::ShowOnlyInnerPropertiesName))
+		{
+			return;
+		}
 	}
 
 	TSharedPtr<SWidget> ValueWidget = SNullWidget::NullWidget;
@@ -1751,7 +1751,9 @@ void FPropertyBagDetails::CustomizeChildren(TSharedRef<IPropertyHandle> StructPr
 		.BagStructProperty = StructProperty,
 		.PropUtils = PropUtils,
 		.bAllowContainers = bAllowContainers,
-		.ChildRowFeatures = ChildRowFeatures,
+		PRAGMA_DISABLE_DEPRECATION_WARNINGS
+		.ChildRowFeatures = bFixedLayout ? EPropertyBagChildRowFeatures::Fixed : ChildRowFeatures
+		PRAGMA_ENABLE_DEPRECATION_WARNINGS
 	};
 
 	// Show the Value (FInstancedStruct) as child rows.
