[StructUtils] Returned the default back to Renaming | Deletion | Deprecated to add back lost functionality and restore the previous behavior

#rnx

diff --git a/Engine/Source/Editor/StructUtilsEditor/Public/PropertyBagDetails.h b/Engine/Source/Editor/StructUtilsEditor/Public/PropertyBagDetails.h
--- a/Engine/Source/Editor/StructUtilsEditor/Public/PropertyBagDetails.h
+++ b/Engine/Source/Editor/StructUtilsEditor/Public/PropertyBagDetails.h
@@ -54,7 +54,7 @@ enum class EPropertyBagChildRowFeatures : uint64
 	Extended = Renaming | Deletion | DragAndDrop | CompactTypeSelector | AccessSpecifierButton | DropDownMenuButton | Menu_Rename | Menu_Delete,
 	All      = AllGeneralOptions | AllMenuOptions,
 	// The default version includes deprecated UI features to support previous behavior.
-	Default = ReadOnly | Deprecated
+	Default = Renaming | Deletion | Deprecated
 };
 
 ENUM_CLASS_FLAGS(EPropertyBagChildRowFeatures);
