Instanced Struct Details: Make CustomizeChildren not depend on CustomizeHeader.
This resolves the issue when the property is customized, so FDetailPropertyRow::OnItemNodeInitialized skips CustomizeHeader.

diff --git a/Engine/Source/Editor/StructUtilsEditor/Private/InstancedStructDetails.cpp b/Engine/Source/Editor/StructUtilsEditor/Private/InstancedStructDetails.cpp
--- a/Engine/Source/Editor/StructUtilsEditor/Private/InstancedStructDetails.cpp
+++ b/Engine/Source/Editor/StructUtilsEditor/Private/InstancedStructDetails.cpp
@@ -433,7 +433,7 @@ void FInstancedStructDetails::OnObjectsReinstanced(const FReplacementObjectMap&
 
 void FInstancedStructDetails::CustomizeChildren(TSharedRef<class IPropertyHandle> StructPropertyHandle, class IDetailChildrenBuilder& StructBuilder, IPropertyTypeCustomizationUtils& StructCustomizationUtils)
 {
-	TSharedRef<FInstancedStructDataDetails> DataDetails = MakeShared<FInstancedStructDataDetails>(StructProperty);
+	TSharedRef<FInstancedStructDataDetails> DataDetails = MakeShared<FInstancedStructDataDetails>(StructPropertyHandle);
 	StructBuilder.AddCustomBuilder(DataDetails);
 }
 
