Instanced Struct Picker: Add AllowedStructs and DisallowedStructs arguments for when the picker is placed manually in a detail panel customization.

diff --git a/Engine/Source/Editor/StructUtilsEditor/Private/SInstancedStructPicker.cpp b/Engine/Source/Editor/StructUtilsEditor/Private/SInstancedStructPicker.cpp
--- a/Engine/Source/Editor/StructUtilsEditor/Private/SInstancedStructPicker.cpp
+++ b/Engine/Source/Editor/StructUtilsEditor/Private/SInstancedStructPicker.cpp
@@ -133,6 +133,8 @@ void SInstancedStructPicker::Construct(const FArguments& InArgs, TSharedPtr<IPro
 	static const FName NAME_StructTypeConst = "StructTypeConst";
 
 	OnStructPicked = InArgs._OnStructPicked;
+	AllowedStructs = InArgs._AllowedStructs;
+	DisallowedStructs = InArgs._DisallowedStructs;
 	StructProperty = MoveTemp(InStructProperty);
 	PropUtils = MoveTemp(InPropertyUtils);
 
@@ -268,6 +270,9 @@ TSharedRef<SWidget> SInstancedStructPicker::GenerateStructPicker()
 		Algo::Transform(PropertyCustomizationHelpers::GetStructsFromMetadataString(StructProperty->GetMetaData("AllowedClasses")), StructFilter->AllowedStructs, SoftPointerTransform);
 		Algo::Transform(PropertyCustomizationHelpers::GetStructsFromMetadataString(StructProperty->GetMetaData("DisallowedClasses")), StructFilter->DisallowedStructs, SoftPointerTransform);
 
+		StructFilter->AllowedStructs.Append(AllowedStructs);
+		StructFilter->DisallowedStructs.Append(DisallowedStructs);
+
 		TArray<UObject*> OwningObjects;
 		StructProperty->GetOuterObjects(OwningObjects);
 		for (UObject* OwningObject : OwningObjects)
@@ -298,6 +303,11 @@ TSharedRef<SWidget> SInstancedStructPicker::GenerateStructPicker()
 			}
 		}
 	}
+	else
+	{
+		StructFilter->AllowedStructs = AllowedStructs;
+		StructFilter->DisallowedStructs = DisallowedStructs;
+	}
 
 	const UScriptStruct* SelectedStruct = nullptr;
 	UE::StructUtils::Private::GetCommonScriptStruct(StructProperty, SelectedStruct);
diff --git a/Engine/Source/Editor/StructUtilsEditor/Public/SInstancedStructPicker.h b/Engine/Source/Editor/StructUtilsEditor/Public/SInstancedStructPicker.h
--- a/Engine/Source/Editor/StructUtilsEditor/Public/SInstancedStructPicker.h
+++ b/Engine/Source/Editor/StructUtilsEditor/Public/SInstancedStructPicker.h
@@ -44,6 +44,8 @@ class STRUCTUTILSEDITOR_API SInstancedStructPicker : public SCompoundWidget
 public:
 	SLATE_BEGIN_ARGS(SInstancedStructPicker) { }
 		SLATE_EVENT(FOnStructPicked, OnStructPicked)
+		SLATE_ARGUMENT(TArray<TSoftObjectPtr<const UScriptStruct>>, AllowedStructs)
+		SLATE_ARGUMENT(TArray<TSoftObjectPtr<const UScriptStruct>>, DisallowedStructs)
 	SLATE_END_ARGS()
 
 	void Construct(const FArguments& InArgs, TSharedPtr<IPropertyHandle> InStructProperty, TSharedPtr<IPropertyUtilities> InPropertyUtils);
@@ -58,6 +60,12 @@ private:
 	/** The base struct that we're allowing to be picked (controlled by the "BaseStruct" meta-data) */
 	TWeakObjectPtr<UScriptStruct> BaseScriptStruct = nullptr;
 
+	/** The array of allowed structs (additionally to the meta-data) */
+	TArray<TSoftObjectPtr<const UScriptStruct>> AllowedStructs;
+
+	/** The array of disallowed structs (additionally to the meta-data) */
+	TArray<TSoftObjectPtr<const UScriptStruct>> DisallowedStructs;
+
 	FText GetDisplayValueString() const;
 	FText GetTooltipText() const;
 	const FSlateBrush* GetDisplayValueIcon() const;
