[StructUtils] Add re-nameable categories for organization in Child Row Details View

+ Includes Drag and Drop
+ Sets up common feature combinations

diff --git a/Engine/Source/Editor/PropertyEditor/Private/SDetailSingleItemRow.cpp b/Engine/Source/Editor/PropertyEditor/Private/SDetailSingleItemRow.cpp
--- a/Engine/Source/Editor/PropertyEditor/Private/SDetailSingleItemRow.cpp
+++ b/Engine/Source/Editor/PropertyEditor/Private/SDetailSingleItemRow.cpp
@@ -459,16 +459,21 @@ void SDetailSingleItemRow::Construct( const FArguments& InArgs, FDetailLayoutCus
 			if (WidgetRow.CustomDragDropHandler)
 			{
 				TSharedPtr<SDetailSingleItemRow> InRow = SharedThis(this);
-				TSharedRef<SWidget> ReorderHandle = PropertyEditorHelpers::MakePropertyReorderHandle(InRow, IsEnabledAttribute);
-				
-				NameColumnBox->AddSlot()
-					.HAlign(HAlign_Left)
-					.VAlign(VAlign_Center)
-					.Padding(-4.0f, 0.0f, -10.0f, 0.0f)
-					.AutoWidth()
-					[
-						ReorderHandle
-					];
+
+				// Let the handler subclasses decide if they want the handle widget. 
+				if (WidgetRow.CustomDragDropHandler->UseHandleWidget())
+				{
+					TSharedRef<SWidget> ReorderHandle = PropertyEditorHelpers::MakePropertyReorderHandle(InRow, IsEnabledAttribute);
+
+					NameColumnBox->AddSlot()
+						.HAlign(HAlign_Left)
+						.VAlign(VAlign_Center)
+						.Padding(-4.0f, 0.0f, -10.0f, 0.0f)
+						.AutoWidth()
+						[
+							ReorderHandle
+						];
+				}
 
 				DragLeaveDelegate = FOnTableRowDragLeave::CreateSP(this, &SDetailSingleItemRow::OnArrayOrCustomDragLeave);
 				AcceptDropDelegate = FOnAcceptDrop::CreateSP(this, &SDetailSingleItemRow::OnCustomAcceptDrop);
diff --git a/Engine/Source/Editor/PropertyEditor/Public/IDetailDragDropHandler.h b/Engine/Source/Editor/PropertyEditor/Public/IDetailDragDropHandler.h
--- a/Engine/Source/Editor/PropertyEditor/Public/IDetailDragDropHandler.h
+++ b/Engine/Source/Editor/PropertyEditor/Public/IDetailDragDropHandler.h
@@ -25,4 +25,7 @@ public:
 	 * The handler can return a different drop zone if desired, or no drop zone if it can't accept the drop.
 	 */
 	virtual TOptional<EItemDropZone> CanAcceptDrop(const FDragDropEvent& DragDropEvent, EItemDropZone DropZone) const = 0;
+
+	/** Determines whether the reorder handle widget is added automatically. */
+	virtual bool UseHandleWidget() const { return true; }
 };
diff --git a/Engine/Source/Editor/StructUtilsEditor/Private/Customizations/PropertyBagDetails.cpp b/Engine/Source/Editor/StructUtilsEditor/Private/Customizations/PropertyBagDetails.cpp
--- a/Engine/Source/Editor/StructUtilsEditor/Private/Customizations/PropertyBagDetails.cpp
+++ b/Engine/Source/Editor/StructUtilsEditor/Private/Customizations/PropertyBagDetails.cpp
@@ -5,9 +5,11 @@
 #include "DetailLayoutBuilder.h"
 #include "DetailWidgetRow.h"
 #include "IDetailChildrenBuilder.h"
+#include "IDetailGroup.h"
 #include "IPropertyUtilities.h"
 #include "PropertyBagDragDropHandler.h"
 #include "ScopedTransaction.h"
+#include "SDraggableBox.h" // Custom widget for drag and drop sections
 #include "StructUtilsMetadata.h"
 #include "STypeSelector.h" // Custom widget for pill type selector
 #include "Framework/MultiBox/MultiBoxBuilder.h"
@@ -429,6 +431,30 @@ void ApplyChangesToPropertyDescs(const FText& SessionName, const TSharedPtr<IPro
 	}
 }
 
+template <typename TFunc>
+void ApplyChangesToSinglePropertyDesc(const FText& SessionName, const TSharedPtr<IPropertyHandle> PropertyHandle, const TSharedPtr<IPropertyHandle>& StructProperty, const TSharedPtr<IPropertyUtilities>& PropUtils, TFunc Function)
+{
+	ApplyChangesToPropertyDescs(SessionName, StructProperty, PropUtils, [Function = std::move(Function), Property = PropertyHandle->GetProperty()](TArray<FPropertyBagPropertyDesc>& PropertyDescs)
+	{
+		if (FPropertyBagPropertyDesc* Desc = PropertyDescs.FindByPredicate([Property](const FPropertyBagPropertyDesc& OutDesc) { return OutDesc.CachedProperty == Property; }))
+		{
+			Function(*Desc);
+		}
+	});
+}
+
+template <typename TFunc>
+void ApplyChangesToSinglePropertyDesc(const FText& SessionName, const FPropertyBagPropertyDesc& PropertyDesc, const TSharedPtr<IPropertyHandle>& StructProperty, const TSharedPtr<IPropertyUtilities>& PropUtils, TFunc Function)
+{
+	ApplyChangesToPropertyDescs(SessionName, StructProperty, PropUtils, [Function = std::move(Function), &PropertyDesc](TArray<FPropertyBagPropertyDesc>& PropertyDescs)
+	{
+		if (FPropertyBagPropertyDesc* Desc = PropertyDescs.FindByPredicate([&PropertyDesc](const FPropertyBagPropertyDesc& OutDesc) { return OutDesc == PropertyDesc; }))
+		{
+			Function(*Desc);
+		}
+	});
+}
+
 bool CanHaveMemberVariableOfType(const FEdGraphPinType& PinType)
 {
 	if (PinType.PinCategory == UEdGraphSchema_K2::PC_Exec 
@@ -966,6 +992,11 @@ TSharedRef<SWidget> CreateTypeSelectionWidget(TSharedPtr<IPropertyHandle> ChildP
 			.Font(IDetailLayoutBuilder::GetDetailFont())
 		];
 }
+
+namespace Constants
+{
+	static constexpr int32 MaxCategoryLength = 255;
+}
 } // UE::StructUtils
 
 //----------------------------------------------------------------//
@@ -1006,6 +1037,233 @@ PRAGMA_ENABLE_DEPRECATION_WARNINGS
 	ensure(PropUtils != nullptr);
 }
 
+void FPropertyBagInstanceDataDetails::OnGroupRowAdded(IDetailGroup& GroupRow, int32 Level, const FString& Category) const
+{
+	FDetailWidgetRow& FolderRow = GroupRow.HeaderRow();
+	TWeakPtr<const FPropertyBagInstanceDataDetails> WeakSelf = SharedThis<const FPropertyBagInstanceDataDetails>(this);
+	FString FullCategoryName = GroupRow.GetGroupName().ToString();
+
+	/*** DRAG AND DROP HANDLER ***/
+	if (EnumHasAllFlags(ChildRowFeatures, EPropertyBagChildRowFeatures::DragAndDrop | EPropertyBagChildRowFeatures::Menu_Categories))
+	{
+		TSharedPtr<FPropertyBagDetailsDragDropHandlerTarget> DragDropHandler = MakeShared<FPropertyBagDetailsDragDropHandlerTarget>();
+
+		DragDropHandler->BindCanAcceptDragDrop(
+			FCanAcceptPropertyBagDetailsRowDropOp::CreateLambda(
+				[FullCategoryName](const TSharedPtr<FPropertyBagDetailsDragDropOp>& DropOp, EItemDropZone DropZone) -> TOptional<EItemDropZone>
+				{
+					if (!DropOp.IsValid() || DropZone != EItemDropZone::OntoItem)
+					{
+						DropOp->SetDecoration(EPropertyBagDropState::Invalid);
+						return TOptional<EItemDropZone>();
+					}
+
+					TOptional<FPropertyBagDetailsDragDropOp::FDecoration> DecorationOverride;
+
+					if (UE::StructUtils::Metadata::AreCategoriesEnabled(DropOp->PropertyDesc)
+						&& UE::StructUtils::Metadata::GetCategory(DropOp->PropertyDesc).Equals(FullCategoryName))
+					{
+						const FSlateBrush* Brush = FAppStyle::Get().GetBrush("Graph.ConnectorFeedback.OKWarn");
+						DecorationOverride.Emplace(LOCTEXT("OnSameCategoryDragDropDecoratorMessage", "Already in this category"), Brush);
+						DropOp->SetDecoration(EPropertyBagDropState::SourceIsTarget, std::move(DecorationOverride));
+						return TOptional<EItemDropZone>();
+					}
+
+					const FSlateBrush* Brush = FAppStyle::Get().GetBrush("Graph.ConnectorFeedback.OK");
+					DecorationOverride.Emplace(LOCTEXT("OnNewCategoryDragDropDecoratorMessage", "Move to this category"), Brush);
+					DropOp->SetDecoration(EPropertyBagDropState::Valid, std::move(DecorationOverride));
+					return DropZone;
+				}));
+
+		DragDropHandler->BindOnHandleDragDrop(
+			FOnPropertyBagDetailsRowDropOp::CreateLambda(
+				[WeakSelf, FullCategoryName, BagStructProperty = BagStructProperty, PropUtils = PropUtils](const FPropertyBagPropertyDesc& DroppedPropertyDesc, const EItemDropZone DropZone) -> FReply
+				{
+					if (ensure(DroppedPropertyDesc.CachedProperty && DropZone == EItemDropZone::OntoItem))
+					{
+						const TSharedPtr<const FPropertyBagInstanceDataDetails> DetailsSP = WeakSelf.Pin();
+						const UPropertyBag* ChildBagStruct = DetailsSP ? UE::StructUtils::Private::GetCommonBagStruct(DetailsSP->BagStructProperty) : nullptr;
+						// Validate these properties are still part of the bag.
+						if (!ChildBagStruct || !ChildBagStruct->FindPropertyDescByProperty(DroppedPropertyDesc.CachedProperty))
+						{
+							return FReply::Unhandled();
+						}
+
+						UE::StructUtils::Private::ApplyChangesToSinglePropertyDesc(
+							LOCTEXT("DragToChangeCategory", "Change property category"),
+							DroppedPropertyDesc,
+							BagStructProperty,
+							PropUtils,
+							[&DroppedPropertyDesc, &FullCategoryName](FPropertyBagPropertyDesc& Desc)
+							{
+								UE::StructUtils::Metadata::SetCategory(Desc, FullCategoryName);
+							});
+
+						return FReply::Handled();
+					}
+
+					return FReply::Unhandled();
+				}));
+
+		// Add the drag and drop handler as a target for the folder row.
+		FolderRow.DragDropHandler(std::move(DragDropHandler));
+	}
+
+	PRAGMA_DISABLE_DEPRECATION_WARNINGS
+	const bool bIsFixed = bFixedLayout || (ChildRowFeatures == EPropertyBagChildRowFeatures::Fixed);
+	PRAGMA_ENABLE_DEPRECATION_WARNINGS
+
+	/*** EDITABLE NAME BLOCK ***/
+	const TSharedPtr<SInlineEditableTextBlock> EditableInlineNameWidget = SNew(SInlineEditableTextBlock)
+		.MultiLine(false)
+		.OverflowPolicy(ETextOverflowPolicy::Ellipsis)
+		.Font(IDetailLayoutBuilder::GetDetailFontBold())
+		.Text(FText::FromString(Category))
+		.OnVerifyTextChanged_Lambda([](const FText& InText, FText& OutErrorMessage)
+		{
+			if (InText.IsEmpty())
+			{
+				OutErrorMessage = LOCTEXT("InlineEmptyCategoryName", "Name is empty");
+				return false;
+			}
+
+			if (InText.ToString().Len() > UE::StructUtils::Constants::MaxCategoryLength)
+			{
+				OutErrorMessage = LOCTEXT("InlineInvalidCategoryName", "Invalid category name");
+				return false;
+			}
+
+			return true;
+		})
+		.OnTextCommitted_Lambda([FullCategoryName, Category = Category, BagStructProperty = BagStructProperty, PropUtils = PropUtils](const FText& InNewText, const ETextCommit::Type InCommitType)
+		{
+			if (InCommitType == ETextCommit::OnEnter || InCommitType == ETextCommit::OnUserMovedFocus)
+			{
+				using namespace UE::StructUtils;
+
+				Private::ApplyChangesToPropertyDescs(
+					LOCTEXT("InlineRenameCategory", "Rename category"),
+					BagStructProperty,
+					PropUtils,
+					[&InNewText, OldCategory = FullCategoryName, &Category](TArray<FPropertyBagPropertyDesc>& PropertyDescs)
+					{
+						FString NewCategory = OldCategory;
+						NewCategory.ReplaceInline(*Category, *InNewText.ToString(), ESearchCase::CaseSensitive);
+						for (FPropertyBagPropertyDesc& Desc : PropertyDescs)
+						{
+							if (Metadata::AreCategoriesEnabled(Desc))
+							{
+								FString DescCategory = Metadata::GetCategory(Desc);
+								if (DescCategory.StartsWith(OldCategory))
+								{
+									DescCategory.ReplaceInline(*OldCategory, *NewCategory);
+									Metadata::SetCategory(Desc, DescCategory);
+								}
+							}
+						}
+					});
+			}
+		})
+		.IsReadOnly(bIsFixed || !EnumHasAnyFlags(ChildRowFeatures, EPropertyBagChildRowFeatures::Renaming));
+
+	/*** CATEGORY NAME AND BUTTONS ***/
+	const TSharedPtr<SBorder> NameContent = SNew(SBorder)
+		.BorderImage(FAppStyle::Get().GetBrush("DetailsView.CategoryMiddle"))
+		.Content()
+		[
+			SNew(SHorizontalBox)
+			+ SHorizontalBox::Slot()
+			.VAlign(VAlign_Center)
+			.AutoWidth()
+			.Padding(1, 0)
+			[
+				EditableInlineNameWidget.ToSharedRef()
+			]
+			+ SHorizontalBox::Slot()
+			.FillWidth(1)
+			[
+				SNew(SSpacer).Size(1)
+			]
+			+ SHorizontalBox::Slot()
+			.HAlign(HAlign_Right)
+			.AutoWidth()
+			[
+				SNew(SButton)
+				.HAlign(HAlign_Center)
+				.VAlign(VAlign_Center)
+				.ButtonStyle(FAppStyle::Get(), "SimpleButton")
+				.ToolTipText(LOCTEXT("DeleteCategory", "Delete this category."))
+				.OnClicked_Lambda([GroupName = GroupRow.GetGroupName(), BagStructProperty = BagStructProperty, PropUtils = PropUtils]()
+				{
+					UE::StructUtils::Private::ApplyChangesToPropertyDescs(
+						LOCTEXT("OnCategoryDeleted", "Delete category"),
+						BagStructProperty,
+						PropUtils,
+						[GroupName](TArray<FPropertyBagPropertyDesc>& PropertyDescs)
+						{
+							for (FPropertyBagPropertyDesc& Desc : PropertyDescs)
+							{
+								if (UE::StructUtils::Metadata::GetCategory(Desc).Equals(GroupName.ToString()))
+								{
+									UE::StructUtils::Metadata::RemoveCategory(Desc);
+								}
+							}
+						});
+
+					return FReply::Handled();
+				})
+				.ButtonStyle(FAppStyle::Get(), "SimpleButton")
+				[
+					SNew(SImage)
+					.DesiredSizeOverride(FVector2D(16))
+					.ColorAndOpacity(FSlateColor::UseForeground())
+					.Image(FAppStyle::Get().GetBrush("Icons.Delete"))
+				]
+			]
+		];
+
+	// Mirrors PropertyEditorConstants::GetRowBackgroundColor, which is private.
+	NameContent->SetBorderBackgroundColor(TAttribute<FSlateColor>::CreateLambda([NameContent, Level]()
+	{
+		int32 ColorIndex = 0;
+		int32 Increment = 1;
+
+		for (int i = 0; i < Level + 1; ++i)
+		{
+			ColorIndex += Increment;
+
+			if (ColorIndex == 0 || ColorIndex == 3)
+			{
+				Increment = -Increment;
+			}
+		}
+
+		static constexpr uint8 ColorOffsets[] =
+		{
+			0, 4, (4 + 2), (6 + 4), (10 + 6)
+		};
+
+		const FSlateColor BaseSlateColor = NameContent->IsHovered() ? FAppStyle::Get().GetSlateColor("Colors.Header") : FAppStyle::Get().GetSlateColor("Colors.Panel");
+
+		const FColor BaseColor = BaseSlateColor.GetSpecifiedColor().ToFColor(true);
+
+		const FColor ColorWithOffset(
+			BaseColor.R + ColorOffsets[ColorIndex],
+			BaseColor.G + ColorOffsets[ColorIndex],
+			BaseColor.B + ColorOffsets[ColorIndex]);
+
+		return FSlateColor(FLinearColor::FromSRGBColor(ColorWithOffset));
+	}));
+
+	FolderRow
+		.ShouldAutoExpand(true)
+		.WholeRowContent()
+		.HAlign(HAlign_Fill)
+		[
+			NameContent.ToSharedRef()
+		];
+}
+
 void FPropertyBagInstanceDataDetails::OnChildRowAdded(IDetailPropertyRow& ChildRow)
 {
 	TSharedPtr<SWidget> NameWidget;
@@ -1139,7 +1397,7 @@ void FPropertyBagInstanceDataDetails::OnChildRowAdded(IDetailPropertyRow& ChildR
 				];
 		}
 
-		// Choose the name text box, editable inline or not.
+		/*** EDITABLE NAME BLOCK ***/
 		TSharedPtr<SInlineEditableTextBlock> EditableInlineNameWidget = SNew(SInlineEditableTextBlock)
 			.IsReadOnly(!EnumHasAnyFlags(ChildRowFeatures, EPropertyBagChildRowFeatures::Renaming))
 			.Font(IDetailLayoutBuilder::GetDetailFont())
@@ -1196,6 +1454,7 @@ void FPropertyBagInstanceDataDetails::OnChildRowAdded(IDetailPropertyRow& ChildR
 					});
 			});
 
+		/*** CURRENT UI AS IT BECOMES DEPRECATED ***/
 		// Deprecated in 5.6 - the combo button on the name widget will be removed in favor of the new drop-down menu.
 		if (EnumHasAnyFlags(ChildRowFeatures, EPropertyBagChildRowFeatures::Deprecated))
 		{
@@ -1243,6 +1502,7 @@ void FPropertyBagInstanceDataDetails::OnChildRowAdded(IDetailPropertyRow& ChildR
 				.Size(1)
 			];
 
+		/*** ACCESS SPECIFIER BUTTON ***/
 		if (EnumHasAnyFlags(ChildRowFeatures, EPropertyBagChildRowFeatures::AccessSpecifierButton))
 		{
 			using namespace UE::StructUtils::Metadata;
@@ -1264,10 +1524,11 @@ void FPropertyBagInstanceDataDetails::OnChildRowAdded(IDetailPropertyRow& ChildR
 							LOCTEXT("OnPropertyAccessSpecifierChanged", "Set access specifier."),
 							BagStructProperty,
 							PropUtils,
-							[Property](TArray<FPropertyBagPropertyDesc>& PropertyDescs)
+							[Property, BagStructProperty](TArray<FPropertyBagPropertyDesc>& PropertyDescs)
 							{
-								if (FPropertyBagPropertyDesc* Desc = PropertyDescs.FindByPredicate([Property](const FPropertyBagPropertyDesc& Desc) { return Desc.CachedProperty == Property; }))
+								if (FPropertyBagPropertyDesc* Desc = PropertyDescs.FindByPredicate([Property, BagStructProperty](const FPropertyBagPropertyDesc& Desc) { return Desc.CachedProperty == Property; }))
 								{
+									BagStructProperty->NotifyPreChange();
 									const bool bIsPrivate = Property->HasAnyPropertyFlags(CPF_NativeAccessSpecifierPrivate | CPF_NativeAccessSpecifierProtected);
 									Desc->PropertyFlags &= ~CPF_NativeAccessSpecifiers;
 									if (bIsPrivate)
@@ -1278,6 +1539,8 @@ void FPropertyBagInstanceDataDetails::OnChildRowAdded(IDetailPropertyRow& ChildR
 									{
 										Desc->PropertyFlags |= CPF_NativeAccessSpecifierPrivate;
 									}
+									BagStructProperty->NotifyPostChange(EPropertyChangeType::ValueSet);
+									BagStructProperty->NotifyFinishedChangingProperties();
 								}
 							});
 
@@ -1306,6 +1569,7 @@ void FPropertyBagInstanceDataDetails::OnChildRowAdded(IDetailPropertyRow& ChildR
 				];
 		}
 
+		/*** DROP-DOWN MENU OPTIONS ***/
 		// Check drop-down is enabled and at least one option as well.
 		if (EnumHasAnyFlags(ChildRowFeatures, EPropertyBagChildRowFeatures::DropDownMenuButton)
 			&& EnumHasAnyFlags(ChildRowFeatures, EPropertyBagChildRowFeatures::AllMenuOptions))
@@ -1315,44 +1579,133 @@ void FPropertyBagInstanceDataDetails::OnChildRowAdded(IDetailPropertyRow& ChildR
 
 			if (EnumHasAnyFlags(ChildRowFeatures, EPropertyBagChildRowFeatures::Menu_TypeSelector))
 			{
-				MenuBuilder.AddWidget(UE::StructUtils::CreateTypeSelectionWidget(ChildPropertyHandle,
-						BagStructProperty, PropUtils,
+				MenuBuilder.BeginSection(/*InExtensionHook=*/NAME_None, LOCTEXT("DropDownMenuSectionTypeSelector", "Type"));
+				MenuBuilder.AddWidget(
+					UE::StructUtils::CreateTypeSelectionWidget(
+						ChildPropertyHandle,
+						BagStructProperty,
+						PropUtils,
 						SPinTypeSelector::ESelectorType::Full,
 						bAllowContainers),
 					/*InLabel=*/FText::GetEmpty());
-				MenuBuilder.AddMenuSeparator();
+				MenuBuilder.EndSection();
 			}
 
-			// Must have property renaming enabled or the editable inline widget will be invalid.
-			if (EnumHasAllFlags(ChildRowFeatures, EPropertyBagChildRowFeatures::Renaming | EPropertyBagChildRowFeatures::Menu_Rename))
+			const bool bMenuRenameEnabled = EnumHasAllFlags(ChildRowFeatures, EPropertyBagChildRowFeatures::Renaming | EPropertyBagChildRowFeatures::Menu_Rename);
+			const bool bMenuDeleteEnabled = EnumHasAllFlags(ChildRowFeatures, EPropertyBagChildRowFeatures::Deletion | EPropertyBagChildRowFeatures::Menu_Delete);
+
+			if (bMenuRenameEnabled | bMenuDeleteEnabled)
 			{
-				MenuBuilder.AddMenuEntry(
-					LOCTEXT("Rename", "Rename"),
-					LOCTEXT("Rename_ToolTip", "Rename property"),
-					FSlateIcon(FAppStyle::GetAppStyleSetName(), "Icons.Edit"),
-					FUIAction(FExecuteAction::CreateLambda([NameWidget = EditableInlineNameWidget]()
-					{
-						if (NameWidget)
+				MenuBuilder.BeginSection(/*InExtensionHook=*/NAME_None, LOCTEXT("DropDownMenuSectionGeneral", "General"));
+
+				// Must have property renaming enabled or the editable inline widget will be invalid.
+				if (bMenuRenameEnabled)
+				{
+					MenuBuilder.AddMenuEntry(
+						LOCTEXT("DropDownMenuRenameProperty", "Rename property"),
+						LOCTEXT("DropDownMenuRenamePropertyToolTip", "Enable the inline property renaming."),
+						FSlateIcon(FAppStyle::GetAppStyleSetName(), "Icons.Edit"),
+						FUIAction(FExecuteAction::CreateLambda([NameWidget = EditableInlineNameWidget]()
 						{
-							NameWidget->EnterEditingMode();
-						}
-					})));
+							if (NameWidget)
+							{
+								NameWidget->EnterEditingMode();
+							}
+						})));
+				}
+
+				if (bMenuDeleteEnabled)
+				{
+					MenuBuilder.AddMenuEntry(
+						LOCTEXT("DropDownMenuRemoveProperty", "Remove property"),
+						LOCTEXT("DropDownMenuRemovePropertyToolTip", "Delete the property."),
+						FSlateIcon(FAppStyle::GetAppStyleSetName(), "Icons.Delete"),
+						FUIAction(FExecuteAction::CreateLambda([BagStructProperty = BagStructProperty, PropUtils = PropUtils, ChildPropertyHandle]()
+						{
+							UE::StructUtils::Private::DeleteProperty(BagStructProperty, ChildPropertyHandle, PropUtils);
+						}))
+					);
+				}
+
+				MenuBuilder.EndSection();
 			}
 
-			if (EnumHasAllFlags(ChildRowFeatures, EPropertyBagChildRowFeatures::Deletion | EPropertyBagChildRowFeatures::Menu_Delete))
+			// The property's category (grouping) can be edited here.
+			if (EnumHasAllFlags(ChildRowFeatures, EPropertyBagChildRowFeatures::Categories | EPropertyBagChildRowFeatures::Menu_Categories))
 			{
-				MenuBuilder.AddMenuEntry(
-					LOCTEXT("Remove", "Remove"),
-					LOCTEXT("Remove_ToolTip", "Remove property"),
-					FSlateIcon(FAppStyle::GetAppStyleSetName(), "Icons.Delete"),
-					FUIAction(FExecuteAction::CreateLambda([BagStructProperty = BagStructProperty, PropUtils = PropUtils, ChildPropertyHandle]()
+				MenuBuilder.BeginSection(/*InExtensionHook=*/NAME_None, LOCTEXT("DropDownMenuSectionCategory", "Category"));
+
+				if (ChildPropertyHandle && ChildPropertyHandle->HasMetaData(UE::StructUtils::Metadata::CategoryName))
+				{
+					MenuBuilder.AddMenuEntry(
+						LOCTEXT("DropDownMenuClearCategory", "Clear category"),
+						LOCTEXT("DropDownMenuClearCategoryToolTip", "Remove the property from its current category."),
+						FSlateIcon(FAppStyle::GetAppStyleSetName(), "Icons.Delete"),
+						FUIAction(FExecuteAction::CreateLambda([ChildPropertyHandle, StructProperty = BagStructProperty, PropUtils = PropUtils]()
+						{
+							UE::StructUtils::Private::ApplyChangesToSinglePropertyDesc(
+								LOCTEXT("DropDownMenuOnCategoryCleared", "Clear property category"),
+								ChildPropertyHandle,
+								StructProperty,
+								PropUtils,
+								[](FPropertyBagPropertyDesc& Desc)
+								{
+									UE::StructUtils::Metadata::RemoveCategory(Desc);
+								});
+						})));
+				}
+
+				// TODO: AddVerifiedEditableText seems to bypass the MenuBuilder Section, so this was added temporarily to force the section to exist.
+				MenuBuilder.AddWidget(SNullWidget::NullWidget, FText::GetEmpty(), false, false);
+
+				MenuBuilder.AddVerifiedEditableText(
+					LOCTEXT("DropDownMenuSubMenuCategoryName", "Category"),
+					LOCTEXT("DropDownMenuSubMenuCategoryTooltip", "Edit this value to change the category of this property. Subcategories can be created with the '|' character."),
+					FSlateIcon(FAppStyle::GetAppStyleSetName(), "Icons.FolderOpen"),
+					TAttribute<FText>::CreateLambda([ChildPropertyHandle]()
+					{
+						FText GroupLabel = FText::FromString("");
+						if (ChildPropertyHandle && ChildPropertyHandle->HasMetaData(UE::StructUtils::Metadata::CategoryName))
+						{
+							GroupLabel = FText::FromString(ChildPropertyHandle->GetMetaData(UE::StructUtils::Metadata::CategoryName));
+						}
+
+						return GroupLabel;
+					}),
+					FOnVerifyTextChanged::CreateLambda([](const FText& InText, FText& OutMessage)
 					{
-						UE::StructUtils::Private::DeleteProperty(BagStructProperty, ChildPropertyHandle, PropUtils);
-					}))
-				);
+						if (InText.ToString().Len() > UE::StructUtils::Constants::MaxCategoryLength)
+						{
+							OutMessage = LOCTEXT("DropDownMenuInvalidCategoryName", "Invalid category name");
+							return false;
+						}
+						else
+						{
+							return true;
+						}
+					}),
+					FOnTextCommitted::CreateLambda([StructProperty = BagStructProperty, PropUtils = PropUtils, ChildPropertyHandle](const FText& CommittedText, const ETextCommit::Type CommitType)
+					{
+						if (CommitType == ETextCommit::OnEnter || CommitType == ETextCommit::OnUserMovedFocus)
+						{
+							UE::StructUtils::Private::ApplyChangesToSinglePropertyDesc(
+								LOCTEXT("DropDownMenuOnCategoryEdited", "Edit property category"),
+								ChildPropertyHandle,
+								StructProperty,
+								PropUtils,
+								[&CommittedText](FPropertyBagPropertyDesc& Desc)
+								{
+									UE::StructUtils::Metadata::SetCategory(Desc, CommittedText.ToString());
+								});
+						}
+					}),
+					FOnTextChanged(),
+					!bEditable);
+
+				MenuBuilder.EndSection();
 			}
 
-					/*** Drop-down arrow menu. ***/
+			/*** DROP-DOWN ARROW MENU ***/
 			PropertyDetailsWidget->AddSlot()
 				.HAlign(HAlign_Right)
 				.AutoWidth()
@@ -1379,46 +1732,137 @@ void FPropertyBagInstanceDataDetails::OnChildRowAdded(IDetailPropertyRow& ChildR
 		/*** DRAG AND DROP HANDLER ***/
 		if (EnumHasAnyFlags(ChildRowFeatures, EPropertyBagChildRowFeatures::DragAndDrop) && PropertyDesc != nullptr)
 		{
-			ChildRow.DragDropHandler(MakeShared<FPropertyBagDetailsDragDropHandler>(*PropertyDesc, FOnPropertyBagDetailsRowDropOp::CreateLambda(
-				[WeakSelf = SharedThis<FPropertyBagInstanceDataDetails>(this).ToWeakPtr(), PropertyDesc = *PropertyDesc](const FPropertyBagPropertyDesc& DroppedPropertyDesc, EItemDropZone DropZone) -> FReply
-				{
-					const TSharedPtr<FPropertyBagInstanceDataDetails> DetailsSP = WeakSelf.Pin();
-					const UPropertyBag* ChildBagStruct = DetailsSP ? UE::StructUtils::Private::GetCommonBagStruct(DetailsSP->BagStructProperty) : nullptr;
-					// Validate these properties are still part of the bag.
-					if (!ChildBagStruct || !ChildBagStruct->FindPropertyDescByProperty(PropertyDesc.CachedProperty) || !ChildBagStruct->FindPropertyDescByProperty(DroppedPropertyDesc.CachedProperty))
+			TSharedPtr<FPropertyBagDetailsDragDropHandler> DragDropHandler = MakeShared<FPropertyBagDetailsDragDropHandler>(*PropertyDesc);
+
+			// Can accept drag and drop check for if this is a valid drop
+			DragDropHandler->BindCanAcceptDragDrop(
+				FCanAcceptPropertyBagDetailsRowDropOp::CreateLambda(
+					[PropertyDesc](const TSharedPtr<FPropertyBagDetailsDragDropOp>& DropOp, EItemDropZone DropZone) -> TOptional<EItemDropZone>
 					{
-						return FReply::Unhandled();
-					}
+						if (!PropertyDesc || !DropOp.IsValid())
+						{
+							DropOp->SetDecoration(EPropertyBagDropState::Invalid);
+							return TOptional<EItemDropZone>();
+						}
 
-					EPropertyBagAlterationResult Result = EPropertyBagAlterationResult::InternalError;
+						if (DropZone == EItemDropZone::OntoItem && PropertyDesc->ID != DropOp->PropertyDesc.ID)
+						{
+							DropOp->SetDecoration(EPropertyBagDropState::Invalid);
+							return TOptional<EItemDropZone>();
+						}
 
-					DetailsSP->BagStructProperty->EnumerateRawData([DropZone, &PropertyDesc, &DroppedPropertyDesc, &Result](void* RawData, const int32 /*DataIndex*/, const int32 /*NumData*/)
+						// No effect to drop in these cases. Either source == target, or moving source above/below target puts source in same location.
+						if (*PropertyDesc == DropOp->PropertyDesc
+							|| (DropZone == EItemDropZone::AboveItem && DropOp->PropertyDesc.GetCachedIndex() == PropertyDesc->GetCachedIndex() - 1)
+							|| (DropZone == EItemDropZone::BelowItem && DropOp->PropertyDesc.GetCachedIndex() == PropertyDesc->GetCachedIndex() + 1))
+						{
+							DropOp->SetDecoration(EPropertyBagDropState::SourceIsTarget);
+							return TOptional<EItemDropZone>();
+						}
+
+						DropOp->SetDecoration(EPropertyBagDropState::Valid);
+						return DropZone;
+					}));
+
+			DragDropHandler->BindOnHandleDragDrop(
+				FOnPropertyBagDetailsRowDropOp::CreateLambda(
+					[WeakSelf, PropertyDesc = *PropertyDesc, BagStructProperty = BagStructProperty, PropUtils = PropUtils](const FPropertyBagPropertyDesc& DroppedPropertyDesc, EItemDropZone DropZone) -> FReply
 					{
-						if (FInstancedPropertyBag* PropertyBag = RawData ? static_cast<FInstancedPropertyBag*>(RawData) : nullptr)
+						using namespace UE::StructUtils;
+
+						const TSharedPtr<FPropertyBagInstanceDataDetails> DetailsSP = WeakSelf.Pin();
+						const UPropertyBag* ChildBagStruct = DetailsSP ? Private::GetCommonBagStruct(DetailsSP->BagStructProperty) : nullptr;
+						// Validate these properties are still part of the bag.
+						if (!ChildBagStruct
+							|| !ChildBagStruct->FindPropertyDescByProperty(PropertyDesc.CachedProperty)
+							|| !ChildBagStruct->FindPropertyDescByProperty(DroppedPropertyDesc.CachedProperty))
 						{
-							Result = PropertyBag->ReorderProperty(DroppedPropertyDesc.Name, PropertyDesc.Name, DropZone == EItemDropZone::AboveItem);
+							return FReply::Unhandled();
 						}
 
-						return true;
-					});
+						EPropertyBagAlterationResult Result = EPropertyBagAlterationResult::InternalError;
 
-					return Result == EPropertyBagAlterationResult::Success ? FReply::Handled() : FReply::Unhandled();
-				})));
-		}
+						DetailsSP->BagStructProperty->EnumerateRawData([DropZone, &PropertyDesc, &DroppedPropertyDesc, &Result](void* RawData, const int32 /*DataIndex*/, const int32 /*NumData*/)
+						{
+							if (FInstancedPropertyBag* PropertyBag = RawData ? static_cast<FInstancedPropertyBag*>(RawData) : nullptr)
+							{
+								Result = PropertyBag->ReorderProperty(DroppedPropertyDesc.Name, PropertyDesc.Name, DropZone == EItemDropZone::AboveItem);
+							}
+
+							return true;
+						});
 
-		// Update the name widget with our new property details composition.
-		NameWidget = PropertyDetailsWidget;
-	} // Fixed layouts will pass through and use the default widgets.
+						if (Result == EPropertyBagAlterationResult::Success)
+						{
+							Private::ApplyChangesToSinglePropertyDesc(
+								LOCTEXT("DragDropReorderProperties", "Reordered properties"),
+								DroppedPropertyDesc,
+								BagStructProperty,
+								PropUtils,
+								[PropertyDesc](FPropertyBagPropertyDesc& Desc)
+								{
+									Metadata::SetCategory(Desc, Metadata::GetCategory(PropertyDesc));
+								});
 
+							return FReply::Handled();
+						}
+						else
+						{
+							return FReply::Unhandled();
+						}
+					}));
+
+			// Bind the drag and drop handler for receiving.
+			ChildRow.DragDropHandler(DragDropHandler);
+
+			// Add draggability for the name widget.
+			NameWidget = SNew(SDraggableBox)
+				.DragDropHandler(DragDropHandler)
+				.RequireDirectHover(true)
+				.Content()
+				[
+					PropertyDetailsWidget
+				];
+
+			// Add draggability for the value widget, maximizing draggable space, but not at the cost of the value widget.
+			PropertyValueWidget = SNew(SDraggableBox)
+				.DragDropHandler(DragDropHandler)
+				.RequireDirectHover(true)
+				.Content()
+				[
+					SNew(SHorizontalBox)
+					+ SHorizontalBox::Slot()
+					.HAlign(HAlign_Left)
+					.AutoWidth()
+					[
+						PropertyValueWidget.ToSharedRef()
+					]
+					+ SHorizontalBox::Slot()
+					.HAlign(HAlign_Fill)
+					.VAlign(VAlign_Fill)
+					.FillWidth(1)
+					[
+						SNullWidget::NullWidget
+					]
+				];
+		}
+		else // Update the name widget with our new property details composition.
+		{
+			NameWidget = PropertyDetailsWidget;
+		}
+	}
+
+	/*** FINAL WIDGET ***/
 	ChildRow
 		.IsEnabled(bEditable)
-		.CustomWidget(/*bShowChildren*/true)
+		.CustomWidget(/*bShowChildren=*/true)
 		.NameContent()
 		.HAlign(HAlign_Fill)
 		[
 			NameWidget.ToSharedRef()
 		]
 		.ValueContent()
+		.HAlign(HAlign_Fill)
 		[
 			PropertyValueWidget.ToSharedRef()
 		];
@@ -1728,9 +2172,14 @@ void FPropertyBagDetails::CustomizeHeader(TSharedRef<IPropertyHandle> StructProp
 	if (!bFixedLayout && ChildRowFeatures != EPropertyBagChildRowFeatures::Fixed)
 	PRAGMA_ENABLE_DEPRECATION_WARNINGS
 	{
-		ValueWidget = MakeAddPropertyWidget(StructProperty, PropUtils, DefaultType);
+		SAssignNew(ValueWidget, SHorizontalBox)
+		+ SHorizontalBox::Slot()
+		.AutoWidth()
+		[
+			MakeAddPropertyWidget(StructProperty, PropUtils, DefaultType).ToSharedRef()
+		];
 	}
-	
+
 	HeaderRow
 		.NameContent()
 		[
diff --git a/Engine/Source/Editor/StructUtilsEditor/Private/Customizations/SDraggableBox.cpp b/Engine/Source/Editor/StructUtilsEditor/Private/Customizations/SDraggableBox.cpp
new file mode 100644
--- /dev/null
+++ b/Engine/Source/Editor/StructUtilsEditor/Private/Customizations/SDraggableBox.cpp
@@ -0,0 +1,51 @@
+// Copyright Epic Games, Inc. All Rights Reserved.
+
+#include "SDraggableBox.h"
+
+#include "IDetailDragDropHandler.h"
+#include "Framework/Application/SlateApplication.h"
+
+void SDraggableBox::Construct(const FArguments& InArgs)
+{
+	DragDropHandler = InArgs._DragDropHandler;
+	bRequireDirectHover = InArgs._RequireDirectHover;
+
+	ChildSlot
+	[
+		InArgs._Content.Widget
+	];
+}
+
+FReply SDraggableBox::OnMouseButtonDown(const FGeometry& MyGeometry, const FPointerEvent& MouseEvent)
+{
+	return FReply::Handled().DetectDrag(SharedThis<SWidget>(this), EKeys::LeftMouseButton);
+}
+
+FReply SDraggableBox::OnDragDetected(const FGeometry& MyGeometry, const FPointerEvent& MouseEvent)
+{
+	if (!bRequireDirectHover || IsDirectlyHovered())
+	{
+		if (DragDropHandler && MouseEvent.IsMouseButtonDown(EKeys::LeftMouseButton))
+		{
+			const TSharedPtr<FDragDropOperation> DragDropOp = DragDropHandler->CreateDragDropOperation();
+			if (DragDropOp.IsValid())
+			{
+				return FReply::Handled().BeginDragDrop(DragDropOp.ToSharedRef());
+			}
+		}
+	}
+
+	return FReply::Unhandled();
+}
+
+FCursorReply SDraggableBox::OnCursorQuery(const FGeometry& MyGeometry, const FPointerEvent& CursorEvent) const
+{
+	if (!bRequireDirectHover || IsDirectlyHovered())
+	{
+		return CursorEvent.IsMouseButtonDown(EKeys::LeftMouseButton) ? FCursorReply::Cursor(EMouseCursor::GrabHandClosed) : FCursorReply::Cursor(EMouseCursor::GrabHand);
+	}
+	else
+	{
+		return FCursorReply::Unhandled();
+	}
+}
diff --git a/Engine/Source/Editor/StructUtilsEditor/Private/Customizations/SDraggableBox.h b/Engine/Source/Editor/StructUtilsEditor/Private/Customizations/SDraggableBox.h
new file mode 100644
--- /dev/null
+++ b/Engine/Source/Editor/StructUtilsEditor/Private/Customizations/SDraggableBox.h
@@ -0,0 +1,29 @@
+// Copyright Epic Games, Inc. All Rights Reserved.
+
+#pragma once
+
+#include "Widgets/SCompoundWidget.h"
+
+class IDetailDragDropHandler;
+
+class SDraggableBox : public SCompoundWidget
+{
+public:
+	SLATE_BEGIN_ARGS(SDraggableBox) {}
+		SLATE_DEFAULT_SLOT(FArguments, Content)
+		SLATE_ARGUMENT(TSharedPtr<IDetailDragDropHandler>, DragDropHandler)
+		SLATE_ARGUMENT(bool, RequireDirectHover)
+	SLATE_END_ARGS()
+
+	void Construct(const FArguments& InArgs);
+
+	//~Begin SWidget interface
+	virtual FReply OnMouseButtonDown(const FGeometry& MyGeometry, const FPointerEvent& MouseEvent) override;
+	virtual FReply OnDragDetected(const FGeometry& MyGeometry, const FPointerEvent& MouseEvent) override;
+	virtual FCursorReply OnCursorQuery(const FGeometry& MyGeometry, const FPointerEvent& CursorEvent) const override;
+	//~End SWidget interface
+
+protected:
+	TSharedPtr<IDetailDragDropHandler> DragDropHandler = nullptr;
+	bool bRequireDirectHover = true;
+};
diff --git a/Engine/Source/Editor/StructUtilsEditor/Private/Customizations/STypeSelector.cpp b/Engine/Source/Editor/StructUtilsEditor/Private/Customizations/STypeSelector.cpp
--- a/Engine/Source/Editor/StructUtilsEditor/Private/Customizations/STypeSelector.cpp
+++ b/Engine/Source/Editor/StructUtilsEditor/Private/Customizations/STypeSelector.cpp
@@ -146,7 +146,7 @@ void STypeSelector::Construct(const FArguments& InArgs, FGetPinTypeTree GetPinTy
 	.OnGetMenuContent(this, &STypeSelector::GetMenuContent, false)
 	.ContentPadding(0.0f)
 	.ToolTipText(this, &STypeSelector::GetToolTipForSelector)
-	.HasDownArrow(true)
+	.HasDownArrow(false)
 	.ButtonStyle(FAppStyle::Get(), "SimpleButton")
 	.ButtonContent()
 	[
@@ -162,7 +162,7 @@ void STypeSelector::Construct(const FArguments& InArgs, FGetPinTypeTree GetPinTy
 		SNew(SWidgetSwitcher)
 		.WidgetIndex_Lambda([this]() { return ReadOnly.Get() ? 1 : 0; })
 		+ SWidgetSwitcher::Slot() // editable version
-		.Padding(SelectorType == ESelectorType::Partial ? FMargin(-6.0f, 0.0f, 0.0f, 0.0f) : FMargin(0))
+		.Padding(-6.0f, 0.0f, 0.0f, 0.0f)
 		[
 			TypeComboButton.ToSharedRef()
 		]
@@ -206,14 +206,15 @@ FReply STypeSelector::OnMouseButtonUp(const FGeometry& MyGeometry, const FPointe
 
 FText STypeSelector::GetToolTipForSelector() const
 {
-	FText EditText;
+	FText TooltipText;
 	if (IsEnabled())
 	{
-		EditText = NSLOCTEXT("STypeSelector", "PinTypeSelectorTooltip", "Left click to select the variable's type. Right click to select a container type.\n");
+		TooltipText = NSLOCTEXT("STypeSelector", "PinTypeSelectorTooltip", "Left click to select the variable's type. Right click to select a container type.");
 	}
 	else
 	{
-		EditText = NSLOCTEXT("STypeSelector", "DisabledPinTypeSelectorTooltip", "Cannot edit variable type when they are inherited from parent.\n");
+		TooltipText = NSLOCTEXT("STypeSelector", "DisabledPinTypeSelectorTooltip", "Cannot edit variable type when they are inherited from parent.");
 	}
-	return EditText;
+
+	return FText::Format(INVTEXT("Type: {0}\n{1}"), GetTypeDescription(/*bIncludeSubcategory=*/false), TooltipText);
 }
diff --git a/Engine/Source/Editor/StructUtilsEditor/Private/InstancedStructDetails.cpp b/Engine/Source/Editor/StructUtilsEditor/Private/InstancedStructDetails.cpp
--- a/Engine/Source/Editor/StructUtilsEditor/Private/InstancedStructDetails.cpp
+++ b/Engine/Source/Editor/StructUtilsEditor/Private/InstancedStructDetails.cpp
@@ -205,7 +205,10 @@ TArray<TWeakObjectPtr<const UStruct>> FInstancedStructDataDetails::GetInstanceTy
 	return Result;
 }
 
-void FInstancedStructDataDetails::GetPropertyGroups(const TArray<TSharedPtr<IPropertyHandle>>& InProperties, IDetailChildrenBuilder& InChildBuilder, TMap<TSharedPtr<IPropertyHandle>, IDetailGroup*>& OutPropertyToGroup) const
+void FInstancedStructDataDetails::GetPropertyGroups(
+	const TArray<TSharedPtr<IPropertyHandle>>& InProperties,
+	IDetailChildrenBuilder& InChildBuilder,
+	TMap<TSharedPtr<IPropertyHandle>, IDetailGroup*>& OutPropertyToGroup) const
 {
 	static const FName CategoryName = FName(TEXT("Category"));
 	static const FName EnableCategoriesName = FName(TEXT("EnableCategories"));
@@ -241,8 +244,9 @@ void FInstancedStructDataDetails::GetPropertyGroups(const TArray<TSharedPtr<IPro
 
 		// For this property, add all of the groups needed for its category (eg, Foo, Foo|Bar, and Foo|Bar|Baz)
 		IDetailGroup* CurrentGroup = nullptr;
-		for (FString& CategoryToAdd : CategoriesToAdd)
+		for (int32 Index = 0; Index < CategoriesToAdd.Num(); ++Index)
 		{
+			FString& CategoryToAdd = CategoriesToAdd[Index];
 			CategoryToAdd.TrimStartAndEndInline();
 
 			// Cover the edge case where there's a category like "Foo|", which is invalid
@@ -266,6 +270,7 @@ void FInstancedStructDataDetails::GetPropertyGroups(const TArray<TSharedPtr<IPro
 					CurrentGroup = &InChildBuilder.AddGroup(FName(CompleteCategory), FText::FromString(CategoryToAdd));
 				}
 
+				OnGroupRowAdded(*CurrentGroup, Index, CategoryToAdd);
 				CategoryToGroup.Add(CompleteCategory, CurrentGroup);
 			}
 			else
diff --git a/Engine/Source/Editor/StructUtilsEditor/Private/PropertyBagDragDropHandler.cpp b/Engine/Source/Editor/StructUtilsEditor/Private/PropertyBagDragDropHandler.cpp
--- a/Engine/Source/Editor/StructUtilsEditor/Private/PropertyBagDragDropHandler.cpp
+++ b/Engine/Source/Editor/StructUtilsEditor/Private/PropertyBagDragDropHandler.cpp
@@ -3,41 +3,73 @@
 #include "PropertyBagDragDropHandler.h"
 
 #include "IPropertyBagEditorGraph.h"
+#include "StructUtilsMetadata.h"
 #include "EdGraph/EdGraph.h"
 #include "Input/Reply.h"
 
+#define LOCTEXT_NAMESPACE "StructUtilsEditor"
+
+namespace PropertyBagDragDrop::Constants
+{
+	const FText DefaultValidTargetText = LOCTEXT("ValidTargetTooltip", "Drop property here");
+	const FText DefaultSourceIsTargetText = LOCTEXT("SourceIsTargetTooltip", "Choose a different target");
+	const FText DefaultInvalidTargetText = LOCTEXT("InvalidTargetTooltip", "Invalid target");
+}
+
+FPropertyBagDetailsDragDropOp::FDecoration::FDecoration(const FText& InMessage, const FSlateBrush* InIcon, const FLinearColor& InColor)
+	: Message(InMessage)
+	, Icon(InIcon)
+	, IconColor(InColor) {}
+
 FPropertyBagDetailsDragDropOp::FPropertyBagDetailsDragDropOp(const FPropertyBagPropertyDesc& InPropertyDesc)
 	: PropertyDesc(InPropertyDesc)
 {
 	FGraphEditorDragDropAction::Construct();
 }
 
-void FPropertyBagDetailsDragDropOp::SetDecoration(const EPropertyBagDropState NewDropState)
+void FPropertyBagDetailsDragDropOp::SetDecoration(const EPropertyBagDropState NewDropState, TOptional<FDecoration> OverriddenDecoration)
 {
 	if (CurrentDropState == NewDropState)
 	{
 		return;
 	}
 
-	if (NewDropState == EPropertyBagDropState::Valid)
+	FText Message;
+	const FSlateBrush* Icon = nullptr;
+	FLinearColor Color = FLinearColor::White;
+
+	if (OverriddenDecoration)
 	{
-		SetSimpleFeedbackMessage(FAppStyle::Get().GetBrush("Graph.ConnectorFeedback.OK"), FLinearColor::White, NSLOCTEXT("StructUtilsEditor", "ValidTargetTooltip", "Drop property here"));
+		Icon = OverriddenDecoration->Icon;
+		Message = OverriddenDecoration->Message;
+		Color = OverriddenDecoration->IconColor;
+		bDropTargetValid = NewDropState == EPropertyBagDropState::Valid;
+	}
+	else if (NewDropState == EPropertyBagDropState::Valid)
+	{
+		Icon = FAppStyle::Get().GetBrush("Graph.ConnectorFeedback.OK");
+		Message = PropertyBagDragDrop::Constants::DefaultValidTargetText;
 		bDropTargetValid = true;
 	}
 	else if (NewDropState == EPropertyBagDropState::SourceIsTarget)
 	{
-		SetSimpleFeedbackMessage(FAppStyle::Get().GetBrush("Graph.ConnectorFeedback.OKWarn"), FLinearColor::White.CopyWithNewOpacity(0.5f), NSLOCTEXT("StructUtilsEditor", "SourceIsTargetTooltip", "Choose a different target"));
+		Icon = FAppStyle::Get().GetBrush("Graph.ConnectorFeedback.OKWarn");
+		Color = FLinearColor::White.CopyWithNewOpacity(0.5f);
+		Message = PropertyBagDragDrop::Constants::DefaultSourceIsTargetText;
 		bDropTargetValid = false;
 	}
 	else // DragDropTarget == EPropertyBagDragDropTargetState::Invalid
 	{
 		check(NewDropState == EPropertyBagDropState::Invalid);
-		SetSimpleFeedbackMessage(FAppStyle::Get().GetBrush("Graph.ConnectorFeedback.Error"), FLinearColor::White, NSLOCTEXT("StructUtilsEditor", "InvalidTargetTooltip", "Invalid target"));
+		Icon = FAppStyle::Get().GetBrush("Graph.ConnectorFeedback.Error");
+		Message = PropertyBagDragDrop::Constants::DefaultSourceIsTargetText;
 		bDropTargetValid = false;
 	}
 
-	CurrentDropState = NewDropState;
+	SetSimpleFeedbackMessage(Icon, Color, Message);
 	SetDecoratorVisibility(true);
+
+	CurrentDropState = NewDropState;
 }
 
 void FPropertyBagDetailsDragDropOp::HoverTargetChanged()
@@ -126,76 +158,64 @@ IPropertyBagEdGraphDragAndDrop* FPropertyBagDetailsDragDropOp::GetPropertyBagEdG
 	return Graph ? Cast<IPropertyBagEdGraphDragAndDrop>(Graph) : nullptr;
 }
 
-FPropertyBagDetailsDragDropHandler::FPropertyBagDetailsDragDropHandler(const FPropertyBagPropertyDesc& InPropertyDesc, FOnPropertyBagDetailsRowDropOp InOnHandleDragDrop)
-	: PropertyDesc(InPropertyDesc)
-	, OnHandleDetailsRowDragDrop(std::move(InOnHandleDragDrop)) {}
+FPropertyBagDetailsDragDropHandlerTarget::FPropertyBagDetailsDragDropHandlerTarget(const FCanAcceptPropertyBagDetailsRowDropOp& CanAcceptDragDrop, const FOnPropertyBagDetailsRowDropOp& OnDragDrop)
+	: CanAcceptDetailsRowDropOp(CanAcceptDragDrop)
+	, OnHandleDetailsRowDropOp(OnDragDrop) {}
 
-TSharedPtr<FDragDropOperation> FPropertyBagDetailsDragDropHandler::CreateDragDropOperation() const
+TSharedPtr<FDragDropOperation> FPropertyBagDetailsDragDropHandlerTarget::CreateDragDropOperation() const
 {
-	if (TSharedPtr<FPropertyBagDetailsDragDropOp> DragOp = MakeShared<FPropertyBagDetailsDragDropOp>(PropertyDesc))
-	{
-		return DragOp;
-	}
-
 	return nullptr;
 }
 
-TOptional<EItemDropZone> FPropertyBagDetailsDragDropHandler::CanAcceptDrop(const FDragDropEvent& DragDropSource, const EItemDropZone DropZone) const
+void FPropertyBagDetailsDragDropHandlerTarget::BindCanAcceptDragDrop(FCanAcceptPropertyBagDetailsRowDropOp&& CanAcceptDragDrop)
 {
-	// Property Bag Details Drag and Drop Handler
-	if (const TSharedPtr<FPropertyBagDetailsDragDropOp> DragOp = DragDropSource.GetOperationAs<FPropertyBagDetailsDragDropOp>())
-	{
-		return CanAcceptPropertyBagDetailsDrop(DragOp, DropZone);
-	}
+	CanAcceptDetailsRowDropOp = std::move(CanAcceptDragDrop);
+}
 
-	return TOptional<EItemDropZone>();
+void FPropertyBagDetailsDragDropHandlerTarget::BindOnHandleDragDrop(FOnPropertyBagDetailsRowDropOp&& OnDragDrop)
+{
+	OnHandleDetailsRowDropOp = std::move(OnDragDrop);
 }
 
-bool FPropertyBagDetailsDragDropHandler::AcceptDrop(const FDragDropEvent& DragDropSource, const EItemDropZone DropZone) const
+TOptional<EItemDropZone> FPropertyBagDetailsDragDropHandlerTarget::CanAcceptDrop(const FDragDropEvent& DragDropSource, const EItemDropZone DropZone) const
 {
 	// Property Bag Details Drag and Drop Handler
-	if (const TSharedPtr<FPropertyBagDetailsDragDropOp> DragOp = DragDropSource.GetOperationAs<FPropertyBagDetailsDragDropOp>())
+	if (const TSharedPtr<FPropertyBagDetailsDragDropOp> DropOp = DragDropSource.GetOperationAs<FPropertyBagDetailsDragDropOp>())
 	{
-		return OnPropertyBagDetailsDrop(DragOp, DropZone).IsEventHandled();
+		if (CanAcceptDetailsRowDropOp.IsBound())
+		{
+			return CanAcceptDetailsRowDropOp.Execute(DropOp, DropZone);
+		}
 	}
 
-	return false;
+	return TOptional<EItemDropZone>();
 }
 
-TOptional<EItemDropZone> FPropertyBagDetailsDragDropHandler::CanAcceptPropertyBagDetailsDrop(const TSharedPtr<FPropertyBagDetailsDragDropOp>& DragOp, EItemDropZone DropZone) const
+bool FPropertyBagDetailsDragDropHandlerTarget::AcceptDrop(const FDragDropEvent& DragDropSource, const EItemDropZone DropZone) const
 {
-	if (!DragOp.IsValid())
-	{
-		DragOp->SetDecoration(EPropertyBagDropState::Invalid);
-		return TOptional<EItemDropZone>();
-	}
-
-	if (DropZone == EItemDropZone::OntoItem && PropertyDesc.ID != DragOp->PropertyDesc.ID)
-	{
-		DragOp->SetDecoration(EPropertyBagDropState::Invalid);
-		return TOptional<EItemDropZone>();
-	}
-
-	// No effect to drop in these cases. Either source == target, or moving source above/below target puts source in same location.
-	if (PropertyDesc == DragOp->PropertyDesc
-		|| (DropZone == EItemDropZone::AboveItem && DragOp->PropertyDesc.GetCachedIndex() == PropertyDesc.GetCachedIndex() - 1)
-		|| (DropZone == EItemDropZone::BelowItem && DragOp->PropertyDesc.GetCachedIndex() == PropertyDesc.GetCachedIndex() + 1))
+	// Property Bag Details Drag and Drop Handler
+	if (const TSharedPtr<FPropertyBagDetailsDragDropOp> DropOp = DragDropSource.GetOperationAs<FPropertyBagDetailsDragDropOp>())
 	{
-		DragOp->SetDecoration(EPropertyBagDropState::SourceIsTarget);
-		return TOptional<EItemDropZone>();
+		if (OnHandleDetailsRowDropOp.IsBound())
+		{
+			return OnHandleDetailsRowDropOp.Execute(DropOp->PropertyDesc, DropZone).IsEventHandled();
+		}
 	}
 
-	DragOp->SetDecoration(EPropertyBagDropState::Valid);
-	return DropZone;
+	return false;
 }
 
-FReply FPropertyBagDetailsDragDropHandler::OnPropertyBagDetailsDrop(const TSharedPtr<FPropertyBagDetailsDragDropOp>& DragOp, const EItemDropZone DropZone) const
+FPropertyBagDetailsDragDropHandler::FPropertyBagDetailsDragDropHandler(const FPropertyBagPropertyDesc& InPropertyDesc)
+	: PropertyDesc(InPropertyDesc) {}
+
+TSharedPtr<FDragDropOperation> FPropertyBagDetailsDragDropHandler::CreateDragDropOperation() const
 {
-	// Currently no functional options for receiving the item onto the property.
-	if (!DragOp.IsValid() || !ensure(DropZone != EItemDropZone::OntoItem) || PropertyDesc == DragOp->PropertyDesc)
+	if (TSharedPtr<FPropertyBagDetailsDragDropOp> DragOp = MakeShared<FPropertyBagDetailsDragDropOp>(PropertyDesc))
 	{
-		return FReply::Unhandled();
+		return DragOp;
 	}
 
-	return OnHandleDetailsRowDragDrop.IsBound() ? OnHandleDetailsRowDragDrop.Execute(DragOp->PropertyDesc, DropZone) : FReply::Unhandled();
+	return nullptr;
 }
+
+#undef LOCTEXT_NAMESPACE
\ No newline at end of file
diff --git a/Engine/Source/Editor/StructUtilsEditor/Private/StructUtilsMetadata.cpp b/Engine/Source/Editor/StructUtilsEditor/Private/StructUtilsMetadata.cpp
--- a/Engine/Source/Editor/StructUtilsEditor/Private/StructUtilsMetadata.cpp
+++ b/Engine/Source/Editor/StructUtilsEditor/Private/StructUtilsMetadata.cpp
@@ -2,16 +2,84 @@
 
 #include "StructUtilsMetadata.h"
 
+#include "StructUtils/PropertyBag.h"
+
 namespace UE::StructUtils::Metadata
 {
+	// For properties
+	const FName EnableCategoriesName = TEXT("EnableCategories");
+	const FName CategoryName = TEXT("Category");
+
+	// For the bag
 	const FName HideInDetailPanelsName = TEXT("HideInDetailPanel");
 	const FName ShowOnlyInnerPropertiesName = TEXT("ShowOnlyInnerProperties");
-
 	const FName FixedLayoutName = TEXT("FixedLayout");
 	const FName DefaultTypeName = TEXT("DefaultType");
 	const FName AllowContainersName = TEXT("AllowContainers");
+	UE_DEPRECATED(5.6, "Use 'AllowContainers' instead.");
 	const FName AllowArraysName = TEXT("AllowArrays");
 	const FName IsPinTypeAcceptedName = TEXT("IsPinTypeAccepted");
 	const FName CanRemovePropertyName = TEXT("CanRemoveProperty");
 	const FName ChildRowFeaturesName = TEXT("ChildRowFeatures");
-}
\ No newline at end of file
+
+	bool AreCategoriesEnabled(const FPropertyBagPropertyDesc& OutDesc)
+	{
+		return OutDesc.MetaData.ContainsByPredicate([](const FPropertyBagPropertyDescMetaData& MetaData) { return MetaData.Key == EnableCategoriesName; });
+	}
+
+	void EnableCategories(FPropertyBagPropertyDesc& OutDesc)
+	{
+		if (!AreCategoriesEnabled(OutDesc))
+		{
+			OutDesc.MetaData.Emplace(FPropertyBagPropertyDescMetaData(EnableCategoriesName, {}));
+		}
+	}
+
+	void DisableCategories(FPropertyBagPropertyDesc& OutDesc)
+	{
+		OutDesc.MetaData.RemoveAllSwap([](const FPropertyBagPropertyDescMetaData& MetaData)
+		{
+			return MetaData.Key == EnableCategoriesName;
+		});
+	}
+
+	void SetCategory(FPropertyBagPropertyDesc& OutDesc, const FString& InGroupLabel, const bool bAutoEnableCategories)
+	{
+		if (FPropertyBagPropertyDescMetaData* MetaData = OutDesc.MetaData.FindByPredicate([](const FPropertyBagPropertyDescMetaData& InMetaData) { return InMetaData.Key == CategoryName; }))
+		{
+			MetaData->Value = InGroupLabel;
+		}
+		else
+		{
+			OutDesc.MetaData.Emplace(FPropertyBagPropertyDescMetaData(CategoryName, InGroupLabel));
+		}
+
+		if (bAutoEnableCategories)
+		{
+			EnableCategories(OutDesc);
+		}
+	}
+
+	void RemoveCategory(FPropertyBagPropertyDesc& OutDesc, const bool bAutoDisableCategories)
+	{
+		OutDesc.MetaData.RemoveAllSwap([](const FPropertyBagPropertyDescMetaData& MetaData)
+		{
+			return MetaData.Key == CategoryName;
+		});
+
+		if (bAutoDisableCategories)
+		{
+			DisableCategories(OutDesc);
+		}
+	}
+
+	FString GetCategory(const FPropertyBagPropertyDesc& InDesc)
+	{
+		if (const FPropertyBagPropertyDescMetaData* MetaData = InDesc.MetaData.FindByPredicate([](const FPropertyBagPropertyDescMetaData& InMetaData) { return InMetaData.Key == CategoryName; }))
+		{
+			return MetaData->Value;
+		}
+
+		return FString("");
+	}
+}
diff --git a/Engine/Source/Editor/StructUtilsEditor/Public/InstancedStructDetails.h b/Engine/Source/Editor/StructUtilsEditor/Public/InstancedStructDetails.h
--- a/Engine/Source/Editor/StructUtilsEditor/Public/InstancedStructDetails.h
+++ b/Engine/Source/Editor/StructUtilsEditor/Public/InstancedStructDetails.h
@@ -68,6 +68,8 @@ public:
 	virtual FName GetName() const override;
 	//~ End IDetailCustomNodeBuilder interface
 
+	// Called when a group is added, override to customize a group row.
+	virtual void OnGroupRowAdded(IDetailGroup& GroupRow, int32 Level, const FString& Category) const {}
 	// Called when a child is added, override to customize a child row.
 	virtual void OnChildRowAdded(IDetailPropertyRow& ChildRow) {}
 
diff --git a/Engine/Source/Editor/StructUtilsEditor/Public/PropertyBagDetails.h b/Engine/Source/Editor/StructUtilsEditor/Public/PropertyBagDetails.h
--- a/Engine/Source/Editor/StructUtilsEditor/Public/PropertyBagDetails.h
+++ b/Engine/Source/Editor/StructUtilsEditor/Public/PropertyBagDetails.h
@@ -28,14 +28,15 @@ UENUM(meta = (Bitflags))
 enum class EPropertyBagChildRowFeatures : uint64
 {
 	/** General Options */
-	Fixed                  = 0,      // Fixed layout. No features enabled.
-	Renaming               = 1ULL << 0, // The property name is editable from the details view.
-	Deletion               = 1ULL << 1, // The property is able to be deleted.
-	DragAndDrop            = 1ULL << 2, // Drag and dropping properties is enabled.
-	CompactTypeSelector    = 1ULL << 3, // A compact type selector widget is enabled to the left of the name.
-	AccessSpecifierButton  = 1ULL << 4, // The property metadata can be set for public/private access.
-	DropDownMenuButton     = 1ULL << 5, // Drop-down menu (down arrow button) is enabled.
-	AllGeneralOptions = Renaming | Deletion | DragAndDrop | CompactTypeSelector | AccessSpecifierButton | DropDownMenuButton,
+	Fixed                 = 0,         // Fixed layout. No features enabled.
+	Renaming              = 1ULL << 0, // The property name is editable from the details view.
+	Deletion              = 1ULL << 1, // The property is able to be deleted.
+	DragAndDrop           = 1ULL << 2, // Drag and dropping properties is enabled.
+	CompactTypeSelector   = 1ULL << 3, // A compact type selector widget is enabled to the left of the name.
+	AccessSpecifierButton = 1ULL << 4, // The property metadata can be set for public/private access.
+	DropDownMenuButton    = 1ULL << 5, // Drop-down menu (down arrow button) is enabled.
+	Categories            = 1ULL << 6, // Categories are enabled (UPROPERTY categories).
+	AllGeneralOptions     = Renaming | Deletion | DragAndDrop | CompactTypeSelector | AccessSpecifierButton | DropDownMenuButton | Categories,
 
 	// Insert new drop general features above this line.
 
@@ -43,15 +44,20 @@ enum class EPropertyBagChildRowFeatures : uint64
 	Menu_TypeSelector = 1ULL << 17, // Adds the type selector pill widget to the drop-down.
 	Menu_Rename       = 1ULL << 18, // Renaming the property from the drop-down menu. Requires property renaming enabled.
 	Menu_Delete       = 1ULL << 19, // Deleting the property from the drop-down menu.
-	AllMenuOptions    = Menu_TypeSelector | Menu_Rename | Menu_Delete,
+	Menu_Categories   = 1ULL << 20, // Create new/remove from categories
+
+	AllMenuOptions = Menu_TypeSelector | Menu_Rename | Menu_Delete | Menu_Categories,
 
 	// Insert new menu features above this line.
 	Deprecated = 1ULL << 63, // To allow for deprecating older features.
 
-	// Below are configurations for convenience.
+	// Below are configurations for convenience. These can be set via the Metadata specifier on the property bag.
 	ReadOnly = Fixed,
-	Core     = Renaming | DropDownMenuButton | AllMenuOptions,
-	Extended = Renaming | Deletion | DragAndDrop | CompactTypeSelector | AccessSpecifierButton | DropDownMenuButton | Menu_Rename | Menu_Delete,
+	// Renaming and deleting enabled, with type selection happening in the drop-down menu.
+	Core     = Renaming | Deletion | DropDownMenuButton | Menu_TypeSelector | Menu_Rename | Menu_Delete,
+	// Also enables the compact type selector icon, drag and drop, and categories support.
+	Extended = Core | DragAndDrop | CompactTypeSelector | Categories | Menu_Categories,
+	// All options.
 	All      = AllGeneralOptions | AllMenuOptions,
 	// The default version includes deprecated UI features to support previous behavior.
 	Default = Renaming | Deletion | Deprecated
@@ -137,6 +143,7 @@ public:
 	FPropertyBagInstanceDataDetails& operator=(FPropertyBagInstanceDataDetails&& Other) noexcept = delete;
 	PRAGMA_ENABLE_DEPRECATION_WARNINGS
 
+	virtual void OnGroupRowAdded(IDetailGroup& GroupRow, int32 Level, const FString& Category) const override;
 	virtual void OnChildRowAdded(IDetailPropertyRow& ChildRow) override;
 	
 	/** Enum describing if a property is overridden, or undetermined (e.g. multiselection) */
diff --git a/Engine/Source/Editor/StructUtilsEditor/Public/PropertyBagDragDropHandler.h b/Engine/Source/Editor/StructUtilsEditor/Public/PropertyBagDragDropHandler.h
--- a/Engine/Source/Editor/StructUtilsEditor/Public/PropertyBagDragDropHandler.h
+++ b/Engine/Source/Editor/StructUtilsEditor/Public/PropertyBagDragDropHandler.h
@@ -6,9 +6,11 @@
 #include "IDetailDragDropHandler.h"
 #include "StructUtils/PropertyBag.h"
 
+class FPropertyBagDetailsDragDropOp;
 class IPropertyBagEdGraphDragAndDrop;
 
-DECLARE_DELEGATE_RetVal_TwoParams(FReply, FOnPropertyBagDetailsRowDropOp, const FPropertyBagPropertyDesc&, EItemDropZone);
+DECLARE_DELEGATE_RetVal_TwoParams(TOptional<EItemDropZone>, FCanAcceptPropertyBagDetailsRowDropOp, TSharedPtr<FPropertyBagDetailsDragDropOp> /*DropOperation*/, EItemDropZone /*DropZone*/);
+DECLARE_DELEGATE_RetVal_TwoParams(FReply, FOnPropertyBagDetailsRowDropOp, const FPropertyBagPropertyDesc& /*PropertyDescription*/, EItemDropZone /*DropZone*/);
 
 /** State of the drop, useful for source-is-target validation on details rows. */
 enum class EPropertyBagDropState
@@ -27,6 +29,15 @@ class STRUCTUTILSEDITOR_API FPropertyBagDetailsDragDropOp : public FGraphEditorD
 public:
 	DRAG_DROP_OPERATOR_TYPE(FPropertyBagDetailsDragDropOp, FGraphEditorDragDropAction);
 
+	struct FDecoration
+	{
+		explicit FDecoration(const FText& InMessage, const FSlateBrush* InIcon = nullptr, const FLinearColor& InColor = FLinearColor::White);
+
+		const FText& Message;
+		const FSlateBrush*& Icon;
+		const FLinearColor& IconColor;
+	};
+
 	/**
 	 * Property Bag Details - Drop Operation Constructor
 	 * @param InPropertyDesc Property description for the dropped property.
@@ -36,8 +47,9 @@ public:
 	/**
 	 * Sets the pop up widget's icon and text.
 	 * @param NewDropState The state of the drop operation.
+	 * @param OverriddenDecoration
 	 */
-	void SetDecoration(EPropertyBagDropState NewDropState);
+	void SetDecoration(EPropertyBagDropState NewDropState, TOptional<FDecoration> OverriddenDecoration = {});
 
 	// ~Begin FGraphEditorDragDropAction interface
 	/** Event when the drag and drop operation changes hover targets. */
@@ -84,23 +96,41 @@ private:
 	IPropertyBagEdGraphDragAndDrop* GetPropertyBagEdGraphDragAndDropInterface() const;
 };
 
-/** Handles drag-and-drop (as both source and target) for a single property's child widget row */
-class STRUCTUTILSEDITOR_API FPropertyBagDetailsDragDropHandler : public IDetailDragDropHandler
+/** Handles drag-and-drop (as a target) for a single property's child widget row. */
+class STRUCTUTILSEDITOR_API FPropertyBagDetailsDragDropHandlerTarget : public IDetailDragDropHandler
 {
 public:
 	/**
-	 * Property Bag Details - Drag/Drop Handler Constructor
-	 * @param InPropertyDesc Property description for the child row.
-	 * @param InOnHandleDragDrop Custom function for handling receipt of a drop operation for this Property Bag details view.
+	 * Property Bag Details - Drag/Drop Handler Default Constructor
 	 */
-	FPropertyBagDetailsDragDropHandler(const FPropertyBagPropertyDesc& InPropertyDesc, FOnPropertyBagDetailsRowDropOp InOnHandleDragDrop);
+	FPropertyBagDetailsDragDropHandlerTarget() = default;
 
 	/**
-	 * Creates the drag and drop operation object that contains necessary information for a drag and drop operation.
-	 * @return Shared pointer to the drag and drop operation.
+	 * Property Bag Details - Drag/Drop Handler Constructor
+	 * @param CanAcceptDragDrop [Signature]: TOptional<EItemDropZone>(TSharedPtr<FPropertyBagDetailsDragDropOp>, EItemDropZone)
+	 * @param OnDragDrop [Signature]: FReply(const FPropertyBagPropertyDesc&, EItemDropZone)
 	 */
+	FPropertyBagDetailsDragDropHandlerTarget(const FCanAcceptPropertyBagDetailsRowDropOp& CanAcceptDragDrop, const FOnPropertyBagDetailsRowDropOp& OnDragDrop);
+
+	/** Creates the drag and drop operation. Disabled for the target. Enabled in subclasses. */
 	virtual TSharedPtr<FDragDropOperation> CreateDragDropOperation() const override;
 
+	/**
+	 * Bind the delegate for determining if this handler can accept this drag and drop operation.
+	 * @param CanAcceptDragDrop [Signature]: TOptional<EItemDropZone>(TSharedPtr<FPropertyBagDetailsDragDropOp>, EItemDropZone)
+	 */
+	virtual void BindCanAcceptDragDrop(FCanAcceptPropertyBagDetailsRowDropOp&& CanAcceptDragDrop);
+
+	/**
+	 * Bind the delegate to handle a valid drag and drop operation.
+	 * @param OnDragDrop [Signature]: FReply(const FPropertyBagPropertyDesc&, EItemDropZone)
+	 */
+	virtual void BindOnHandleDragDrop(FOnPropertyBagDetailsRowDropOp&& OnDragDrop);
+
+	/** Disable automatic creation of the handle widget for targets. It is overridden in the source below. */
+	virtual bool UseHandleWidget() const override { return false; }
+
+protected:
 	/**
 	 * Whether the details view child row can accept a drag and drop operation. 
 	 * @param DragDropSource The source dragged onto this row.
@@ -118,15 +148,40 @@ public:
 	virtual bool AcceptDrop(const FDragDropEvent& DragDropSource, EItemDropZone DropZone) const override;
 
 private:
-	/** CanAcceptDrop helper for Property Bag Details row cases, specifically. */
-	TOptional<EItemDropZone> CanAcceptPropertyBagDetailsDrop(const TSharedPtr<FPropertyBagDetailsDragDropOp>& DragOp, EItemDropZone DropZone) const;
+	/**
+	 * Callback to determine if the target is acceptable for a details view child row drag and drop operation.
+	 * Signature: Optional<EItemDropZone>(TSharedPtr<FPropertyBagDetailsDragDropOp>, EItemDropZone)
+	 */
+	FCanAcceptPropertyBagDetailsRowDropOp CanAcceptDetailsRowDropOp;
 
-	/** AcceptDrop helper for Property Bag Details row cases, specifically. */
-	FReply OnPropertyBagDetailsDrop(const TSharedPtr<FPropertyBagDetailsDragDropOp>& DragOp, EItemDropZone DropZone) const;
+	/**
+	 * Callback to determine if the drag to a details view child row was handled successfully.
+	 * Signature: FReply(const FPropertyBagPropertyDesc&, EItemDropZone)
+	 */
+	FOnPropertyBagDetailsRowDropOp OnHandleDetailsRowDropOp;
+};
 
-	/** The current row's property bag property description. */
-	const FPropertyBagPropertyDesc PropertyDesc;
+/** Handles drag-and-drop (as a source or target) for a single property's child widget row. */
+class STRUCTUTILSEDITOR_API FPropertyBagDetailsDragDropHandler : public FPropertyBagDetailsDragDropHandlerTarget
+{
+public:
+	/**
+	 * Constructor.
+	 * @param InPropertyDesc The property desc of the property represented by this child widget row.
+	 */
+	FPropertyBagDetailsDragDropHandler(const FPropertyBagPropertyDesc& InPropertyDesc);
+
+protected:
+	/**
+	 * Creates the drag and drop operation object that contains necessary information for a drag and drop operation.
+	 * @return Shared pointer to the drag and drop operation.
+	 */
+	virtual TSharedPtr<FDragDropOperation> CreateDragDropOperation() const override;
+
+	/** Enable the automatic creation of the handle (grip) widget for the source. */
+	virtual bool UseHandleWidget() const override { return true; }
 
-	/** Callback for widget to determine if the drag to a details view child row was handled successfully. Signature: FReply(const FPropertyBagPropertyDesc&, EItemDropZone)*/
-	FOnPropertyBagDetailsRowDropOp OnHandleDetailsRowDragDrop;
+private:
+	/** The current child row's property bag property description. */
+	const FPropertyBagPropertyDesc PropertyDesc;
 };
diff --git a/Engine/Source/Editor/StructUtilsEditor/Public/StructUtilsMetadata.h b/Engine/Source/Editor/StructUtilsEditor/Public/StructUtilsMetadata.h
--- a/Engine/Source/Editor/StructUtilsEditor/Public/StructUtilsMetadata.h
+++ b/Engine/Source/Editor/StructUtilsEditor/Public/StructUtilsMetadata.h
@@ -4,35 +4,54 @@
 
 #include "UObject/NameTypes.h"
 
+struct FPropertyBagPropertyDesc;
+
 // Metadata used by StructUtils
 namespace UE::StructUtils::Metadata
 {
 	// Metadata usable in UPROPERTY for customizing the behavior when displaying the property in a property panel or graph node
 
-	/// Common UPROPERTY metadata names.
+	/// FInstancedPropertyBag
 	extern STRUCTUTILSEDITOR_API const FName HideInDetailPanelsName;
 	extern STRUCTUTILSEDITOR_API const FName ShowOnlyInnerPropertiesName;
 
-	/// FInstancedPropertyBag
-
-	/// [PropertyMetadata] FixedLayout: Indicates that the instanced property bag has a fixed layout and it is not possible to add/remove properties.
+	/// [Bag Property Metadata] FixedLayout: Indicates that the instanced property bag has a fixed layout and it is not possible to add/remove properties.
 	extern STRUCTUTILSEDITOR_API const FName FixedLayoutName;
 
-	/// [PropertyMetadata] DefaultType: Default property type when adding a new Property. Should be taken from enum EPropertyBagPropertyType
+	/// [Bag Property Metadata] DefaultType: Default property type when adding a new Property. Should be taken from enum EPropertyBagPropertyType
 	extern STRUCTUTILSEDITOR_API const FName DefaultTypeName;
 
-	/// [PropertyMetadata] AllowContainers: By default it is always allowed to create containers (Array, Set) for properties. Use AllowContainers=false to disable container support.
+	/// [Bag Property Metadata] AllowContainers: By default it is always allowed to create containers (Array, Set) for properties. Use AllowContainers=false to disable container support.
 	extern STRUCTUTILSEDITOR_API const FName AllowContainersName;
 
 	UE_DEPRECATED(5.6, "Remains for backwards compatibility. Please see 'AllowContainersName'.")
 	extern STRUCTUTILSEDITOR_API const FName AllowArraysName;
 
-	/// [PropertyMetadata] IsPinTypeAccepted: Name of a UFunction with signature bool(FEdGraphPinType). Returns false if the type should be discarded.
+	/// [Bag Property Metadata] IsPinTypeAccepted: Name of a UFunction with signature bool(FEdGraphPinType). Returns false if the type should be discarded.
 	extern STRUCTUTILSEDITOR_API const FName IsPinTypeAcceptedName;
 
-	/// [PropertyMetadata] CanRemoveProperty: Name of a UFunction with signature bool(FGuid, FName). ID and name of the property that will be removed. Returns false if the property should not be removed.
+	/// [Bag Property Metadata] CanRemoveProperty: Name of a UFunction with signature bool(FGuid, FName). ID and name of the property that will be removed. Returns false if the property should not be removed.
 	extern STRUCTUTILSEDITOR_API const FName CanRemovePropertyName;
 
-	/// [PropertyMetadata] ChildRowFeatures: A list of UI features available to the PropertyBag properties when displayed in the details view child rows.
+	/// [Bag Property Metadata] ChildRowFeatures: A list of UI features available to the PropertyBag properties when displayed in the details view child rows.
 	extern STRUCTUTILSEDITOR_API const FName ChildRowFeaturesName;
+
+	/// Common property metadata specifiers
+	/// [Property Metadata] EnableCategoriesName: Required to enable the use of categories in general. Per property.
+	extern STRUCTUTILSEDITOR_API const FName EnableCategoriesName;
+	/// [Property Metadata] CategoryName: The name of the category (grouping of properties).
+	extern STRUCTUTILSEDITOR_API const FName CategoryName;
+
+	/** If organizing by category is enabled for this property. */
+	extern STRUCTUTILSEDITOR_API bool AreCategoriesEnabled(const FPropertyBagPropertyDesc& OutDesc);
+	/** Enable category organization for this property. */
+	extern STRUCTUTILSEDITOR_API void EnableCategories(FPropertyBagPropertyDesc& OutDesc);
+	/** Disable category organization for this property. */
+	extern STRUCTUTILSEDITOR_API void DisableCategories(FPropertyBagPropertyDesc& OutDesc);
+	/** Set the category for this property. */
+	extern STRUCTUTILSEDITOR_API void SetCategory(FPropertyBagPropertyDesc& OutDesc, const FString& InGroupLabel, bool bAutoEnableCategories = true);
+	/** Remove the category for this property. */
+	extern STRUCTUTILSEDITOR_API void RemoveCategory(FPropertyBagPropertyDesc& OutDesc, bool bAutoDisableCategories = true);
+	/** Get the category for this property. */
+	extern STRUCTUTILSEDITOR_API FString GetCategory(const FPropertyBagPropertyDesc& InDesc);
 }
