[StructUtils] Fix build issues

* Fix binding reference to FSlateBrush pointer
* Add namespace for StructUtils SDraggableBox to not conflict with Modeling Tools SDraggableBox

#rnx

diff --git a/Engine/Source/Editor/StructUtilsEditor/Private/Customizations/PropertyBagDetails.cpp b/Engine/Source/Editor/StructUtilsEditor/Private/Customizations/PropertyBagDetails.cpp
--- a/Engine/Source/Editor/StructUtilsEditor/Private/Customizations/PropertyBagDetails.cpp
+++ b/Engine/Source/Editor/StructUtilsEditor/Private/Customizations/PropertyBagDetails.cpp
@@ -1816,7 +1816,7 @@ void FPropertyBagInstanceDataDetails::OnChildRowAdded(IDetailPropertyRow& ChildR
 			ChildRow.DragDropHandler(DragDropHandler);
 
 			// Add draggability for the name widget.
-			NameWidget = SNew(SDraggableBox)
+			NameWidget = SNew(StructUtilsEditor::SDraggableBox)
 				.DragDropHandler(DragDropHandler)
 				.RequireDirectHover(true)
 				.Content()
@@ -1825,7 +1825,7 @@ void FPropertyBagInstanceDataDetails::OnChildRowAdded(IDetailPropertyRow& ChildR
 				];
 
 			// Add draggability for the value widget, maximizing draggable space, but not at the cost of the value widget.
-			PropertyValueWidget = SNew(SDraggableBox)
+			PropertyValueWidget = SNew(StructUtilsEditor::SDraggableBox)
 				.DragDropHandler(DragDropHandler)
 				.RequireDirectHover(true)
 				.Content()
diff --git a/Engine/Source/Editor/StructUtilsEditor/Private/Customizations/SDraggableBox.cpp b/Engine/Source/Editor/StructUtilsEditor/Private/Customizations/SDraggableBox.cpp
--- a/Engine/Source/Editor/StructUtilsEditor/Private/Customizations/SDraggableBox.cpp
+++ b/Engine/Source/Editor/StructUtilsEditor/Private/Customizations/SDraggableBox.cpp
@@ -5,47 +5,50 @@
 #include "IDetailDragDropHandler.h"
 #include "Framework/Application/SlateApplication.h"
 
-void SDraggableBox::Construct(const FArguments& InArgs)
+namespace StructUtilsEditor
 {
-	DragDropHandler = InArgs._DragDropHandler;
-	bRequireDirectHover = InArgs._RequireDirectHover;
+	void SDraggableBox::Construct(const FArguments& InArgs)
+	{
+		DragDropHandler = InArgs._DragDropHandler;
+		bRequireDirectHover = InArgs._RequireDirectHover;
 
-	ChildSlot
-	[
-		InArgs._Content.Widget
-	];
-}
+		ChildSlot
+		[
+			InArgs._Content.Widget
+		];
+	}
 
-FReply SDraggableBox::OnMouseButtonDown(const FGeometry& MyGeometry, const FPointerEvent& MouseEvent)
-{
-	return FReply::Handled().DetectDrag(SharedThis<SWidget>(this), EKeys::LeftMouseButton);
-}
+	FReply SDraggableBox::OnMouseButtonDown(const FGeometry& MyGeometry, const FPointerEvent& MouseEvent)
+	{
+		return FReply::Handled().DetectDrag(SharedThis<SWidget>(this), EKeys::LeftMouseButton);
+	}
 
-FReply SDraggableBox::OnDragDetected(const FGeometry& MyGeometry, const FPointerEvent& MouseEvent)
-{
-	if (!bRequireDirectHover || IsDirectlyHovered())
+	FReply SDraggableBox::OnDragDetected(const FGeometry& MyGeometry, const FPointerEvent& MouseEvent)
 	{
-		if (DragDropHandler && MouseEvent.IsMouseButtonDown(EKeys::LeftMouseButton))
+		if (!bRequireDirectHover || IsDirectlyHovered())
 		{
-			const TSharedPtr<FDragDropOperation> DragDropOp = DragDropHandler->CreateDragDropOperation();
-			if (DragDropOp.IsValid())
+			if (DragDropHandler && MouseEvent.IsMouseButtonDown(EKeys::LeftMouseButton))
 			{
-				return FReply::Handled().BeginDragDrop(DragDropOp.ToSharedRef());
+				const TSharedPtr<FDragDropOperation> DragDropOp = DragDropHandler->CreateDragDropOperation();
+				if (DragDropOp.IsValid())
+				{
+					return FReply::Handled().BeginDragDrop(DragDropOp.ToSharedRef());
+				}
 			}
 		}
-	}
-
-	return FReply::Unhandled();
-}
 
-FCursorReply SDraggableBox::OnCursorQuery(const FGeometry& MyGeometry, const FPointerEvent& CursorEvent) const
-{
-	if (!bRequireDirectHover || IsDirectlyHovered())
-	{
-		return CursorEvent.IsMouseButtonDown(EKeys::LeftMouseButton) ? FCursorReply::Cursor(EMouseCursor::GrabHandClosed) : FCursorReply::Cursor(EMouseCursor::GrabHand);
+		return FReply::Unhandled();
 	}
-	else
+
+	FCursorReply SDraggableBox::OnCursorQuery(const FGeometry& MyGeometry, const FPointerEvent& CursorEvent) const
 	{
-		return FCursorReply::Unhandled();
+		if (!bRequireDirectHover || IsDirectlyHovered())
+		{
+			return CursorEvent.IsMouseButtonDown(EKeys::LeftMouseButton) ? FCursorReply::Cursor(EMouseCursor::GrabHandClosed) : FCursorReply::Cursor(EMouseCursor::GrabHand);
+		}
+		else
+		{
+			return FCursorReply::Unhandled();
+		}
 	}
-}
+}
\ No newline at end of file
diff --git a/Engine/Source/Editor/StructUtilsEditor/Private/Customizations/SDraggableBox.h b/Engine/Source/Editor/StructUtilsEditor/Private/Customizations/SDraggableBox.h
--- a/Engine/Source/Editor/StructUtilsEditor/Private/Customizations/SDraggableBox.h
+++ b/Engine/Source/Editor/StructUtilsEditor/Private/Customizations/SDraggableBox.h
@@ -6,10 +6,12 @@
 
 class IDetailDragDropHandler;
 
-class SDraggableBox : public SCompoundWidget
+namespace StructUtilsEditor
 {
-public:
-	SLATE_BEGIN_ARGS(SDraggableBox) {}
+	class SDraggableBox : public SCompoundWidget
+	{
+	public:
+		SLATE_BEGIN_ARGS(SDraggableBox) {}
 		SLATE_DEFAULT_SLOT(FArguments, Content)
 		SLATE_ARGUMENT(TSharedPtr<IDetailDragDropHandler>, DragDropHandler)
 		SLATE_ARGUMENT(bool, RequireDirectHover)
@@ -17,13 +19,14 @@ public:
 
 	void Construct(const FArguments& InArgs);
 
-	//~Begin SWidget interface
-	virtual FReply OnMouseButtonDown(const FGeometry& MyGeometry, const FPointerEvent& MouseEvent) override;
-	virtual FReply OnDragDetected(const FGeometry& MyGeometry, const FPointerEvent& MouseEvent) override;
-	virtual FCursorReply OnCursorQuery(const FGeometry& MyGeometry, const FPointerEvent& CursorEvent) const override;
-	//~End SWidget interface
+		//~Begin SWidget interface
+		virtual FReply OnMouseButtonDown(const FGeometry& MyGeometry, const FPointerEvent& MouseEvent) override;
+		virtual FReply OnDragDetected(const FGeometry& MyGeometry, const FPointerEvent& MouseEvent) override;
+		virtual FCursorReply OnCursorQuery(const FGeometry& MyGeometry, const FPointerEvent& CursorEvent) const override;
+		//~End SWidget interface
 
-protected:
-	TSharedPtr<IDetailDragDropHandler> DragDropHandler = nullptr;
-	bool bRequireDirectHover = true;
-};
+	protected:
+		TSharedPtr<IDetailDragDropHandler> DragDropHandler = nullptr;
+		bool bRequireDirectHover = true;
+	};
+}
\ No newline at end of file
diff --git a/Engine/Source/Editor/StructUtilsEditor/Public/PropertyBagDragDropHandler.h b/Engine/Source/Editor/StructUtilsEditor/Public/PropertyBagDragDropHandler.h
--- a/Engine/Source/Editor/StructUtilsEditor/Public/PropertyBagDragDropHandler.h
+++ b/Engine/Source/Editor/StructUtilsEditor/Public/PropertyBagDragDropHandler.h
@@ -34,7 +34,7 @@ public:
 		explicit FDecoration(const FText& InMessage, const FSlateBrush* InIcon = nullptr, const FLinearColor& InColor = FLinearColor::White);
 
 		const FText& Message;
-		const FSlateBrush*& Icon;
+		const FSlateBrush* Icon;
 		const FLinearColor& IconColor;
 	};
 
