[StructUtils] Property Bag Details - Adjust Category naming to restrict to 70 (valid) characters

#rnx

diff --git a/Engine/Source/Editor/StructUtilsEditor/Private/Customizations/PropertyBagDetails.cpp b/Engine/Source/Editor/StructUtilsEditor/Private/Customizations/PropertyBagDetails.cpp
--- a/Engine/Source/Editor/StructUtilsEditor/Private/Customizations/PropertyBagDetails.cpp
+++ b/Engine/Source/Editor/StructUtilsEditor/Private/Customizations/PropertyBagDetails.cpp
@@ -995,7 +995,9 @@ TSharedRef<SWidget> CreateTypeSelectionWidget(TSharedPtr<IPropertyHandle> ChildP
 
 namespace Constants
 {
-	static constexpr int32 MaxCategoryLength = 255;
+	static constexpr int32 MaxCategoryLength = 70;
+	// Special case for categories. Alphanumeric, but including spaces and `|` for nested categories.
+	static constexpr TCHAR InvalidCategoryCharacters[] = TEXT("\"',/.:&!?~\\\n\r\t@#(){}[]<>=;^%$`*+-");
 }
 } // UE::StructUtils
 
@@ -1039,6 +1041,8 @@ PRAGMA_ENABLE_DEPRECATION_WARNINGS
 
 void FPropertyBagInstanceDataDetails::OnGroupRowAdded(IDetailGroup& GroupRow, int32 Level, const FString& Category) const
 {
+	using namespace UE::StructUtils;
+
 	FDetailWidgetRow& FolderRow = GroupRow.HeaderRow();
 	TWeakPtr<const FPropertyBagInstanceDataDetails> WeakSelf = SharedThis<const FPropertyBagInstanceDataDetails>(this);
 	FString FullCategoryName = GroupRow.GetGroupName().ToString();
@@ -1060,8 +1064,8 @@ void FPropertyBagInstanceDataDetails::OnGroupRowAdded(IDetailGroup& GroupRow, in
 
 					TOptional<FPropertyBagDetailsDragDropOp::FDecoration> DecorationOverride;
 
-					if (UE::StructUtils::Metadata::AreCategoriesEnabled(DropOp->PropertyDesc)
-						&& UE::StructUtils::Metadata::GetCategory(DropOp->PropertyDesc).Equals(FullCategoryName))
+					if (Metadata::AreCategoriesEnabled(DropOp->PropertyDesc)
+						&& Metadata::GetCategory(DropOp->PropertyDesc).Equals(FullCategoryName))
 					{
 						const FSlateBrush* Brush = FAppStyle::Get().GetBrush("Graph.ConnectorFeedback.OKWarn");
 						DecorationOverride.Emplace(LOCTEXT("OnSameCategoryDragDropDecoratorMessage", "Already in this category"), Brush);
@@ -1082,21 +1086,21 @@ void FPropertyBagInstanceDataDetails::OnGroupRowAdded(IDetailGroup& GroupRow, in
 					if (ensure(DroppedPropertyDesc.CachedProperty && DropZone == EItemDropZone::OntoItem))
 					{
 						const TSharedPtr<const FPropertyBagInstanceDataDetails> DetailsSP = WeakSelf.Pin();
-						const UPropertyBag* ChildBagStruct = DetailsSP ? UE::StructUtils::Private::GetCommonBagStruct(DetailsSP->BagStructProperty) : nullptr;
+						const UPropertyBag* ChildBagStruct = DetailsSP ? Private::GetCommonBagStruct(DetailsSP->BagStructProperty) : nullptr;
 						// Validate these properties are still part of the bag.
 						if (!ChildBagStruct || !ChildBagStruct->FindPropertyDescByProperty(DroppedPropertyDesc.CachedProperty))
 						{
 							return FReply::Unhandled();
 						}
 
-						UE::StructUtils::Private::ApplyChangesToSinglePropertyDesc(
+						Private::ApplyChangesToSinglePropertyDesc(
 							LOCTEXT("DragToChangeCategory", "Change property category"),
 							DroppedPropertyDesc,
 							BagStructProperty,
 							PropUtils,
 							[&DroppedPropertyDesc, &FullCategoryName](FPropertyBagPropertyDesc& Desc)
 							{
-								UE::StructUtils::Metadata::SetCategory(Desc, FullCategoryName);
+								Metadata::SetCategory(Desc, FullCategoryName);
 							});
 
 						return FReply::Handled();
@@ -1126,10 +1130,14 @@ void FPropertyBagInstanceDataDetails::OnGroupRowAdded(IDetailGroup& GroupRow, in
 				OutErrorMessage = LOCTEXT("InlineEmptyCategoryName", "Name is empty");
 				return false;
 			}
-
-			if (InText.ToString().Len() > UE::StructUtils::Constants::MaxCategoryLength)
+			else if (InText.ToString().Len() > Constants::MaxCategoryLength)
 			{
-				OutErrorMessage = LOCTEXT("InlineInvalidCategoryName", "Invalid category name");
+				OutErrorMessage = LOCTEXT("InlineInvalidCategoryLength", "Too many characters");
+				return false;
+			}
+			else if (!FName::IsValidXName(InText.ToString(), Constants::InvalidCategoryCharacters))
+			{
+				OutErrorMessage = LOCTEXT("InlineInvalidCategoryName", "Invalid character(s)");
 				return false;
 			}
 
@@ -1195,7 +1203,7 @@ void FPropertyBagInstanceDataDetails::OnGroupRowAdded(IDetailGroup& GroupRow, in
 				.ToolTipText(LOCTEXT("DeleteCategory", "Delete this category."))
 				.OnClicked_Lambda([GroupName = GroupRow.GetGroupName(), BagStructProperty = BagStructProperty, PropUtils = PropUtils]()
 				{
-					UE::StructUtils::Private::ApplyChangesToPropertyDescs(
+					Private::ApplyChangesToPropertyDescs(
 						LOCTEXT("OnCategoryDeleted", "Delete category"),
 						BagStructProperty,
 						PropUtils,
@@ -1203,9 +1211,9 @@ void FPropertyBagInstanceDataDetails::OnGroupRowAdded(IDetailGroup& GroupRow, in
 						{
 							for (FPropertyBagPropertyDesc& Desc : PropertyDescs)
 							{
-								if (UE::StructUtils::Metadata::GetCategory(Desc).Equals(GroupName.ToString()))
+								if (Metadata::GetCategory(Desc).Equals(GroupName.ToString()))
 								{
-									UE::StructUtils::Metadata::RemoveCategory(Desc);
+									Metadata::RemoveCategory(Desc);
 								}
 							}
 						});
