[PCG] Fix bug related to STypeSelector retaining previous options each time

#rnx

diff --git a/Engine/Source/Editor/StructUtilsEditor/Private/Customizations/STypeSelector.cpp b/Engine/Source/Editor/StructUtilsEditor/Private/Customizations/STypeSelector.cpp
--- a/Engine/Source/Editor/StructUtilsEditor/Private/Customizations/STypeSelector.cpp
+++ b/Engine/Source/Editor/StructUtilsEditor/Private/Customizations/STypeSelector.cpp
@@ -180,7 +180,7 @@ FReply STypeSelector::OnMouseButtonDown(const FGeometry& MyGeometry, const FPoin
 		bIsRightMousePressed = true;
 	}
 
-	return SPinTypeSelector::OnMouseButtonDown(MyGeometry, MouseEvent);
+	return FReply::Unhandled();
 }
 
 FReply STypeSelector::OnMouseButtonUp(const FGeometry& MyGeometry, const FPointerEvent& MouseEvent)
@@ -200,10 +200,17 @@ FReply STypeSelector::OnMouseButtonUp(const FGeometry& MyGeometry, const FPointe
 	else
 	{
 		bIsRightMousePressed = false;
-		return SPinTypeSelector::OnMouseButtonUp(MyGeometry, MouseEvent);
+		return FReply::Unhandled();
 	}
 }
 
+TSharedRef<SWidget> STypeSelector::GetMenuContent(const bool bForSecondaryType)
+{
+	// Reset the TypeTreeRoot in case property type filters or funcs have changed.
+	TypeTreeRoot.Reset();
+	return SPinTypeSelector::GetMenuContent(bForSecondaryType);
+}
+
 FText STypeSelector::GetToolTipForSelector() const
 {
 	FText TooltipText;
diff --git a/Engine/Source/Editor/StructUtilsEditor/Private/Customizations/STypeSelector.h b/Engine/Source/Editor/StructUtilsEditor/Private/Customizations/STypeSelector.h
--- a/Engine/Source/Editor/StructUtilsEditor/Private/Customizations/STypeSelector.h
+++ b/Engine/Source/Editor/StructUtilsEditor/Private/Customizations/STypeSelector.h
@@ -52,6 +52,8 @@ public:
 	virtual FReply OnMouseButtonUp(const FGeometry& MyGeometry, const FPointerEvent& MouseEvent) override;
 	// End of SWidget interface
 
+	virtual TSharedRef<SWidget> GetMenuContent(bool bForSecondaryType) override;
+
 private:
 	FText GetToolTipForSelector() const;
 };
