[StructUtils] Cleanup of recent UI changes for release.

#rnx

diff --git a/Engine/Source/Editor/StructUtilsEditor/Private/StructUtilsMetadata.cpp b/Engine/Source/Editor/StructUtilsEditor/Private/StructUtilsMetadata.cpp
--- a/Engine/Source/Editor/StructUtilsEditor/Private/StructUtilsMetadata.cpp
+++ b/Engine/Source/Editor/StructUtilsEditor/Private/StructUtilsMetadata.cpp
@@ -7,20 +7,20 @@
 namespace UE::StructUtils::Metadata
 {
 	// For properties
-	const FName EnableCategoriesName = TEXT("EnableCategories");
-	const FName CategoryName = TEXT("Category");
+	const FLazyName EnableCategoriesName = "EnableCategories";
+	const FLazyName CategoryName = "Category";
 
 	// For the bag
-	const FName HideInDetailPanelsName = TEXT("HideInDetailPanel");
-	const FName ShowOnlyInnerPropertiesName = TEXT("ShowOnlyInnerProperties");
-	const FName FixedLayoutName = TEXT("FixedLayout");
-	const FName DefaultTypeName = TEXT("DefaultType");
-	const FName AllowContainersName = TEXT("AllowContainers");
+	const FLazyName HideInDetailPanelsName = "HideInDetailPanel";
+	const FLazyName ShowOnlyInnerPropertiesName = "ShowOnlyInnerProperties";
+	const FLazyName FixedLayoutName = "FixedLayout";
+	const FLazyName DefaultTypeName = "DefaultType";
+	const FLazyName AllowContainersName = "AllowContainers";
 	UE_DEPRECATED(5.6, "Use 'AllowContainers' instead.");
-	const FName AllowArraysName = TEXT("AllowArrays");
-	const FName IsPinTypeAcceptedName = TEXT("IsPinTypeAccepted");
-	const FName CanRemovePropertyName = TEXT("CanRemoveProperty");
-	const FName ChildRowFeaturesName = TEXT("ChildRowFeatures");
+	const FLazyName AllowArraysName = "AllowArrays";
+	const FLazyName IsPinTypeAcceptedName = "IsPinTypeAccepted";
+	const FLazyName CanRemovePropertyName = "CanRemoveProperty";
+	const FLazyName ChildRowFeaturesName = "ChildRowFeatures";
 
 	bool AreCategoriesEnabled(const FPropertyBagPropertyDesc& OutDesc)
 	{
diff --git a/Engine/Source/Editor/StructUtilsEditor/Public/PropertyBagDragDropHandler.h b/Engine/Source/Editor/StructUtilsEditor/Public/PropertyBagDragDropHandler.h
--- a/Engine/Source/Editor/StructUtilsEditor/Public/PropertyBagDragDropHandler.h
+++ b/Engine/Source/Editor/StructUtilsEditor/Public/PropertyBagDragDropHandler.h
@@ -33,9 +33,9 @@ public:
 	{
 		explicit FDecoration(const FText& InMessage, const FSlateBrush* InIcon = nullptr, const FLinearColor& InColor = FLinearColor::White);
 
-		const FText& Message;
+		const FText Message;
 		const FSlateBrush* Icon;
-		const FLinearColor& IconColor;
+		const FLinearColor IconColor;
 	};
 
 	/**
@@ -47,7 +47,7 @@ public:
 	/**
 	 * Sets the pop up widget's icon and text.
 	 * @param NewDropState The state of the drop operation.
-	 * @param OverriddenDecoration
+	 * @param OverriddenDecoration An optional overridden decorator for the UI, including Message, Icon, and IconColor
 	 */
 	void SetDecoration(EPropertyBagDropState NewDropState, TOptional<FDecoration> OverriddenDecoration = {});
 
@@ -128,7 +128,10 @@ public:
 	virtual void BindOnHandleDragDrop(FOnPropertyBagDetailsRowDropOp&& OnDragDrop);
 
 	/** Disable automatic creation of the handle widget for targets. It is overridden in the source below. */
-	virtual bool UseHandleWidget() const override { return false; }
+	virtual bool UseHandleWidget() const override
+	{
+		return false;
+	}
 
 protected:
 	/**
diff --git a/Engine/Source/Editor/StructUtilsEditor/Public/StructUtilsMetadata.h b/Engine/Source/Editor/StructUtilsEditor/Public/StructUtilsMetadata.h
--- a/Engine/Source/Editor/StructUtilsEditor/Public/StructUtilsMetadata.h
+++ b/Engine/Source/Editor/StructUtilsEditor/Public/StructUtilsMetadata.h
@@ -12,35 +12,35 @@ namespace UE::StructUtils::Metadata
 	// Metadata usable in UPROPERTY for customizing the behavior when displaying the property in a property panel or graph node
 
 	/// FInstancedPropertyBag
-	extern STRUCTUTILSEDITOR_API const FName HideInDetailPanelsName;
-	extern STRUCTUTILSEDITOR_API const FName ShowOnlyInnerPropertiesName;
+	extern STRUCTUTILSEDITOR_API const FLazyName HideInDetailPanelsName;
+	extern STRUCTUTILSEDITOR_API const FLazyName ShowOnlyInnerPropertiesName;
 
 	/// [Bag Property Metadata] FixedLayout: Indicates that the instanced property bag has a fixed layout and it is not possible to add/remove properties.
-	extern STRUCTUTILSEDITOR_API const FName FixedLayoutName;
+	extern STRUCTUTILSEDITOR_API const FLazyName FixedLayoutName;
 
 	/// [Bag Property Metadata] DefaultType: Default property type when adding a new Property. Should be taken from enum EPropertyBagPropertyType
-	extern STRUCTUTILSEDITOR_API const FName DefaultTypeName;
+	extern STRUCTUTILSEDITOR_API const FLazyName DefaultTypeName;
 
 	/// [Bag Property Metadata] AllowContainers: By default it is always allowed to create containers (Array, Set) for properties. Use AllowContainers=false to disable container support.
-	extern STRUCTUTILSEDITOR_API const FName AllowContainersName;
+	extern STRUCTUTILSEDITOR_API const FLazyName AllowContainersName;
 
 	UE_DEPRECATED(5.6, "Remains for backwards compatibility. Please see 'AllowContainersName'.")
-	extern STRUCTUTILSEDITOR_API const FName AllowArraysName;
+	extern STRUCTUTILSEDITOR_API const FLazyName AllowArraysName;
 
 	/// [Bag Property Metadata] IsPinTypeAccepted: Name of a UFunction with signature bool(FEdGraphPinType). Returns false if the type should be discarded.
-	extern STRUCTUTILSEDITOR_API const FName IsPinTypeAcceptedName;
+	extern STRUCTUTILSEDITOR_API const FLazyName IsPinTypeAcceptedName;
 
 	/// [Bag Property Metadata] CanRemoveProperty: Name of a UFunction with signature bool(FGuid, FName). ID and name of the property that will be removed. Returns false if the property should not be removed.
-	extern STRUCTUTILSEDITOR_API const FName CanRemovePropertyName;
+	extern STRUCTUTILSEDITOR_API const FLazyName CanRemovePropertyName;
 
 	/// [Bag Property Metadata] ChildRowFeatures: A list of UI features available to the PropertyBag properties when displayed in the details view child rows.
-	extern STRUCTUTILSEDITOR_API const FName ChildRowFeaturesName;
+	extern STRUCTUTILSEDITOR_API const FLazyName ChildRowFeaturesName;
 
 	/// Common property metadata specifiers
 	/// [Property Metadata] EnableCategoriesName: Required to enable the use of categories in general. Per property.
-	extern STRUCTUTILSEDITOR_API const FName EnableCategoriesName;
+	extern STRUCTUTILSEDITOR_API const FLazyName EnableCategoriesName;
 	/// [Property Metadata] CategoryName: The name of the category (grouping of properties).
-	extern STRUCTUTILSEDITOR_API const FName CategoryName;
+	extern STRUCTUTILSEDITOR_API const FLazyName CategoryName;
 
 	/** If organizing by category is enabled for this property. */
 	extern STRUCTUTILSEDITOR_API bool AreCategoriesEnabled(const FPropertyBagPropertyDesc& OutDesc);
