[PCG] Stabilization efforts for StructUtils' PropertyBagDetails related to StateMachine crash.

* Changed UI to warn on invalid name characters for properties
* Added some robustness to renaming, sanitizing names

#rnx

diff --git a/Engine/Source/Editor/StructUtilsEditor/Private/Customizations/PropertyBagDetails.cpp b/Engine/Source/Editor/StructUtilsEditor/Private/Customizations/PropertyBagDetails.cpp
--- a/Engine/Source/Editor/StructUtilsEditor/Private/Customizations/PropertyBagDetails.cpp
+++ b/Engine/Source/Editor/StructUtilsEditor/Private/Customizations/PropertyBagDetails.cpp
@@ -1284,11 +1284,6 @@ void FPropertyBagInstanceDataDetails::OnChildRowAdded(IDetailPropertyRow& ChildR
 
 	TWeakPtr<FPropertyBagInstanceDataDetails> WeakSelf = SharedThis<FPropertyBagInstanceDataDetails>(this);
 
-	bool bSupportedType = false;
-	bool bEditable = false;
-
-	FText UnsupportedTypeWarning;
-
 	const UPropertyBag* BagStruct = BagStructProperty ? UE::StructUtils::Private::GetCommonBagStruct(BagStructProperty) : nullptr;
 	const FProperty* ChildProperty = ChildPropertyHandle->GetProperty();
 	const FPropertyBagPropertyDesc* PropertyDesc = BagStruct ? BagStruct->FindPropertyDescByProperty(ChildProperty) : nullptr;
@@ -1303,24 +1298,28 @@ void FPropertyBagInstanceDataDetails::OnChildRowAdded(IDetailPropertyRow& ChildR
 		ChildRow.Visibility(EVisibility::Collapsed);
 		return;
 	}
-	else if (!ensure(PropertyDesc) || PropertyDesc->ContainerTypes.Num() > 1)
-	{
-		static const FText UnsupportedTypeWarningNestedContainer = LOCTEXT("NestedContainersWarning", "This property type is not supported in the property bag UI.");
 
+	bool bEditable = BagStructProperty->IsEditable();
+
+	/*** WARNINGS FOR PROPERTY ISSUES ***/
+	FText WarningOnProperty; // This message will supplement a warning icon on the details view child row, which will show if not empty.
+	if (!ensure(PropertyDesc) || PropertyDesc->ContainerTypes.Num() > 1)
+	{
 		// The property editing for nested containers is not supported.
-		UnsupportedTypeWarning = UnsupportedTypeWarningNestedContainer;
+		WarningOnProperty = LOCTEXT("NestedContainersWarning", "This property type (nested container) is not supported in the property bag UI.");
 	}
 	else if ((PropertyDesc->ValueType == EPropertyBagPropertyType::UInt32 || PropertyDesc->ValueType == EPropertyBagPropertyType::UInt64) && !bIsFixed)
 	{
-		static const FText UnsupportedTypeWarningUnsigned = LOCTEXT("UnsignedTypesWarning", "Unsigned types are not supported through the property type selection. If you change the type, you will not be able to change it back.");
-
 		// Warn that the unsigned types cannot be set via the type selection.
-		UnsupportedTypeWarning = UnsupportedTypeWarningUnsigned;
+		WarningOnProperty = LOCTEXT("UnsignedTypesWarning", "Unsigned types are not supported through the property type selection. If you change the type, you will not be able to change it back.");
+	}
+	else if (UE::StructUtils::Private::HasMissingType(ChildPropertyHandle))
+	{
+		WarningOnProperty = LOCTEXT("MissingTypeWarning", "The property is missing type. The Struct, Enum, or Object may have been removed.");
 	}
-	else // The property is valid - continue.
+	else if (!FInstancedPropertyBag::IsPropertyNameValid(UE::StructUtils::Private::GetPropertyNameSafe(ChildPropertyHandle)))
 	{
-		bEditable = BagStructProperty->IsEditable();
-		bSupportedType = true;
+		WarningOnProperty = LOCTEXT("InvalidNameWarning", "The property's name contains invalid characters. Dynamically named properties with invalid characters may be rejected in future releases.");
 	}
 
 	/*** OVERRIDE RESET TO DEFAULT ACTION FOR BAG OVERRIDES ***/
@@ -1380,13 +1379,10 @@ void FPropertyBagInstanceDataDetails::OnChildRowAdded(IDetailPropertyRow& ChildR
 				.VAlign(VAlign_Center)
 				[
 					SNew(SImage)
-					.ToolTipText_Lambda([UnsupportedTypeWarning, bSupportedType]()
+					.ToolTipText(WarningOnProperty)
+					.Visibility_Lambda([WarningOnProperty]()
 					{
-						return !bSupportedType ? UnsupportedTypeWarning : LOCTEXT("MissingType", "The property is missing type. The Struct, Enum, or Object may have been removed.");
-					})
-					.Visibility_Lambda([ChildPropertyHandle, bSupportedType]()
-					{
-						return UE::StructUtils::Private::HasMissingType(ChildPropertyHandle) || !bSupportedType ? EVisibility::Visible : EVisibility::Collapsed;
+						return WarningOnProperty.IsEmpty() ? EVisibility::Collapsed : EVisibility::Visible;
 					})
 					.DesiredSizeOverride(FVector2D(12))
 					.ColorAndOpacity(FLinearColor(1.f, 0.8f, 0.f, 1.f))
@@ -1413,24 +1409,49 @@ void FPropertyBagInstanceDataDetails::OnChildRowAdded(IDetailPropertyRow& ChildR
 			.OverflowPolicy(ETextOverflowPolicy::Ellipsis)
 			.Text_Lambda([ChildPropertyHandle]()
 			{
-				return FText::FromName(UE::StructUtils::Private::GetPropertyNameSafe(ChildPropertyHandle));
+				const bool bIsBool = ChildPropertyHandle && ChildPropertyHandle->GetPropertyClass() == FBoolProperty::StaticClass();
+				const FName PropertyName = UE::StructUtils::Private::GetPropertyNameSafe(ChildPropertyHandle);
+				return FText::FromString(FName::NameToDisplayString(PropertyName.ToString(), bIsBool));
 			})
 			.OnVerifyTextChanged_Lambda([BagStructProperty = BagStructProperty, ChildPropertyHandle](const FText& InText, FText& OutErrorMessage)
 			{
 				if (InText.IsEmpty())
 				{
-					OutErrorMessage = LOCTEXT("EmptyName", "Name is empty");
+					OutErrorMessage = LOCTEXT("InlineEmptyPropertyName", "Name is empty");
+					return false;
+				}
+
+				// Check for invalid characters upon renaming.
+				if (!FInstancedPropertyBag::IsPropertyNameValid(InText.ToString()))
+				{
+					OutErrorMessage = LOCTEXT("InlineInvalidPropertyName", "Invalid character(s)");
 					return false;
 				}
 
-				const FName NewName = FInstancedPropertyBag::SanitizePropertyName(InText.ToString());
 				const FName OldName = UE::StructUtils::Private::GetPropertyNameSafe(ChildPropertyHandle);
-				const bool bResult = UE::StructUtils::Private::IsUniqueName(NewName, OldName, BagStructProperty);
-				if (!bResult)
+				// Bypass if the name is the exact same.
+				if (InText.ToString().Equals(OldName.ToString()))
 				{
-					OutErrorMessage = LOCTEXT("MustBeUniqueName", "Property must have unique name");
+					return true;
 				}
-				return bResult;
+
+				// Sanitize out any other characters that we allowed for convenience but are not valid, like spaces.
+				const FName NewName = FInstancedPropertyBag::SanitizePropertyName(InText.ToString());
+
+				// Bypass if sanitized name is the same.
+				if (NewName == OldName)
+				{
+					return true;
+				}
+
+				if (!UE::StructUtils::Private::IsUniqueName(NewName, OldName, BagStructProperty))
+				{
+					OutErrorMessage = LOCTEXT("InlinePropertyUniqueName", "Property must have unique name");
+					return false;
+				}
+
+				// Name is OK.
+				return true;
 			})
 			.OnTextCommitted_Lambda([BagStructProperty = BagStructProperty, PropUtils = PropUtils, ChildPropertyHandle](const FText& InNewText, ETextCommit::Type InCommitType)
 			{
@@ -1441,15 +1462,14 @@ void FPropertyBagInstanceDataDetails::OnChildRowAdded(IDetailPropertyRow& ChildR
 
 				const FName NewName = FInstancedPropertyBag::SanitizePropertyName(InNewText.ToString());
 				const FName OldName = UE::StructUtils::Private::GetPropertyNameSafe(ChildPropertyHandle);
-				if (!UE::StructUtils::Private::IsUniqueName(NewName, OldName, BagStructProperty))
+
+				if (!ensureMsgf(UE::StructUtils::Private::IsUniqueName(NewName, OldName, BagStructProperty), TEXT("Should have already been addressed in OnVerifyTextChanged.")))
 				{
 					return;
 				}
 
 				UE::StructUtils::Private::ApplyChangesToPropertyDescs(
-					FText::Format(LOCTEXT("OnPropertyNameChanged", "Change property name: {0} -> {1}"),
-						FText::FromName(OldName),
-						FText::FromName(NewName)),
+					FText::Format(LOCTEXT("OnPropertyNameChanged", "Change property name: {0} -> {1}"), FText::FromName(OldName), FText::FromName(NewName)),
 					BagStructProperty,
 					PropUtils,
 					[&NewName, &ChildPropertyHandle](TArray<FPropertyBagPropertyDesc>& PropertyDescs)
diff --git a/Engine/Source/Runtime/CoreUObject/Private/StructUtils/PropertyBag.cpp b/Engine/Source/Runtime/CoreUObject/Private/StructUtils/PropertyBag.cpp
--- a/Engine/Source/Runtime/CoreUObject/Private/StructUtils/PropertyBag.cpp
+++ b/Engine/Source/Runtime/CoreUObject/Private/StructUtils/PropertyBag.cpp
@@ -1567,6 +1567,18 @@ namespace UE::StructUtils::Private
 
 		return EPropertyBagAlterationResult::Success;
 	}
+
+	namespace Constants
+	{
+		/**
+		 * Minimal set of invalid characters that must be removed for a name to be valid. This allows convenience
+		 * characters (ex. spaces) in display names for better readability, but stricter sanitization can and should
+		 * be applied later.
+		 * Derived from INVALID_OBJECTNAME_CHARACTERS.
+		 */
+		constexpr TCHAR MinimalInvalidNameCharacters[] = TEXT("\"'/.:&!?~\\\n\r\t@#(){}[]<>=;^%$`");
+		constexpr TCHAR InvalidNameCharacters[] = TEXT(" \"',/.:|&!?~\\\n\r\t@#(){}[]<>=;^%$`+*");
+	}
 } // UE::StructUtils::Private
 
 
@@ -1893,7 +1905,8 @@ EPropertyBagAlterationResult FInstancedPropertyBag::AddProperties(const TConstAr
 		}
 		else if (NewDesc.ValueType != EPropertyBagPropertyType::None)
 		{
-			Descs.Add(NewDesc);
+			FPropertyBagPropertyDesc& Desc = Descs.Add_GetRef(NewDesc);
+			Desc.Name = SanitizePropertyName(NewDesc.Name, /*ReplacementChar=*/TEXT('_'));
 		}
 	}
 
@@ -3292,21 +3305,40 @@ void FInstancedPropertyBag::GetPreloadDependencies(TArray<UObject*>& OutDeps)
 	}
 }
 
-FName FInstancedPropertyBag::SanitizePropertyName(const FString& Name)
+bool FInstancedPropertyBag::IsPropertyNameValid(const FString& Name)
+{
+	return !Name.IsEmpty() && FName::IsValidXName(Name, UE::StructUtils::Private::Constants::MinimalInvalidNameCharacters);
+}
+
+bool FInstancedPropertyBag::IsPropertyNameValid(const FName Name)
+{
+	return IsPropertyNameValid(Name.ToString());
+}
+
+FName FInstancedPropertyBag::SanitizePropertyName(const FString& Name, const TCHAR ReplacementChar)
 {
-	if (!IsPropertyNameValid(Name))
+	if (Name.IsEmpty())
 	{
-		return !Name.IsEmpty() ? MakeObjectNameFromDisplayLabel(Name, NAME_None) : FName(TEXT("Property"));
+		return FName("Property");
 	}
-	else
+
+	FString SanitizedName = Name;
+	constexpr TStringView InvalidChars = UE::StructUtils::Private::Constants::InvalidNameCharacters;
+
+	// Validate the replacement character is not on the invalid list.
+	checkf(!InvalidChars.Contains(FString(1, &ReplacementChar)), TEXT("Invalid replacement character."));
+
+	for (const TCHAR Char : InvalidChars)
 	{
-		return FName(Name);
+		SanitizedName.ReplaceCharInline(Char, ReplacementChar);
 	}
+
+	return FName(SanitizedName);
 }
 
-FName FInstancedPropertyBag::SanitizePropertyName(const FName Name)
+FName FInstancedPropertyBag::SanitizePropertyName(const FName Name, const TCHAR ReplacementChar)
 {
-	return SanitizePropertyName(Name.ToString());
+	return SanitizePropertyName(Name.ToString(), ReplacementChar);
 }
 
 //----------------------------------------------------------------//
diff --git a/Engine/Source/Runtime/CoreUObject/Public/StructUtils/PropertyBag.h b/Engine/Source/Runtime/CoreUObject/Public/StructUtils/PropertyBag.h
--- a/Engine/Source/Runtime/CoreUObject/Public/StructUtils/PropertyBag.h
+++ b/Engine/Source/Runtime/CoreUObject/Public/StructUtils/PropertyBag.h
@@ -861,38 +861,38 @@ struct FInstancedPropertyBag
 	UE_API void GetPreloadDependencies(TArray<UObject*>& OutDeps);
 
 	/**
-	 * Checks whether a provided name is a valid property bag property name.
+	 * Checks whether a provided name is a valid property bag name to create a new property bag property.
+	 * Note: Some characters are allowed that are still invalid, but for workflow reasons they are acceptable and
+	 * should still be sanitized to remove them when adding to the property bag, ex. spaces.
 	 * @param Name The name to check for invalid characters.
-	 * @return Sanitized property name based on the input string.
+	 * @return True if the property name is void of any strictly invalid characters.
 	 */
-	static bool IsPropertyNameValid(const FString& Name)
-	{
-		return !Name.IsEmpty() && FName::IsValidXName(Name, INVALID_OBJECTNAME_CHARACTERS);
-	}
+	UE_API static bool IsPropertyNameValid(const FString& Name);
 
 	/**
-	 * Checks whether a provided name is a valid property bag property name.
+	 * Checks whether a provided name is a valid property bag name to create a new property bag property.
+	 * Note: Some characters are allowed that are still invalid, but for workflow reasons they are acceptable and
+	 * should still be sanitized to remove them when adding to the property bag, ex. spaces.
 	 * @param Name The name to check for invalid characters.
-	 * @return Sanitized property name based on the input name.
+	 * @return True if the property name is void of any strictly invalid characters.
 	 */
-	static bool IsPropertyNameValid(FName Name)
-	{
-		return !Name.IsNone() && FName::IsValidXName(Name, INVALID_OBJECTNAME_CHARACTERS);
-	}
+	UE_API static bool IsPropertyNameValid(const FName Name);
 
 	/**
 	 * Returns a sanitized version of the provided name without invalid characters.
 	 * @param Name The name to sanitize.
-	 * @return Sanitized property name based on the input string.
+	 * @param ReplacementChar Will replace invalid characters.
+	 * @return Potentially modified property name replacing invalid characters with the provided replacement.
 	 */
-	UE_API static FName SanitizePropertyName(const FString& Name);
+	UE_API static FName SanitizePropertyName(const FString& Name, const TCHAR ReplacementChar = TEXT('_'));
 
 	/**
 	 * Returns a sanitized version of the provided name without invalid characters.
 	 * @param Name The name to sanitize.
-	 * @return Sanitized property name based on the input name.
+	 * @param ReplacementChar Will replace invalid characters.
+	 * @return Potentially modified property name replacing invalid characters with the provided replacement.
 	 */
-	UE_API static FName SanitizePropertyName(FName Name);
+	UE_API static FName SanitizePropertyName(FName Name, const TCHAR ReplacementChar = TEXT('_'));
 
 protected:
 	UE_API const void* GetValueAddress(const FPropertyBagPropertyDesc* Desc) const;
