Added a custom paste implementation to the Instanced Struct detail customization, to ensure that the pasted data conforms to the type restrictions imposed by the property before allowing the paste

diff --git a/Engine/Source/Editor/StructUtilsEditor/Private/InstancedStructDetails.cpp b/Engine/Source/Editor/StructUtilsEditor/Private/InstancedStructDetails.cpp
--- a/Engine/Source/Editor/StructUtilsEditor/Private/InstancedStructDetails.cpp
+++ b/Engine/Source/Editor/StructUtilsEditor/Private/InstancedStructDetails.cpp
@@ -5,12 +5,14 @@
 #include "DetailLayoutBuilder.h"
 #include "Editor.h"
 #include "Editor/EditorEngine.h"
+#include "HAL/PlatformApplicationMisc.h"
 #include "IDetailChildrenBuilder.h"
 #include "IDetailGroup.h"
 #include "IPropertyUtilities.h"
 #include "UObject/Package.h"
 #include "Widgets/Images/SImage.h"
 #include "Modules/ModuleManager.h"
+#include "ScopedTransaction.h"
 #include "StructUtils/UserDefinedStruct.h"
 #include "StructUtils/InstancedStruct.h"
 #include "IStructureDataProvider.h"
@@ -414,6 +416,8 @@ void FInstancedStructDetails::CustomizeHeader(TSharedRef<class IPropertyHandle>
 	OnObjectsReinstancedHandle = FCoreUObjectDelegates::OnObjectsReinstanced.AddSP(this, &FInstancedStructDetails::OnObjectsReinstanced);
 
 	HeaderRow
+		.CopyAction(FUIAction(FExecuteAction::CreateSP(this, &FInstancedStructDetails::OnCopyInstancedStruct)))
+		.PasteAction(FUIAction(FExecuteAction::CreateSP(this, &FInstancedStructDetails::OnPasteInstancedStruct), FCanExecuteAction::CreateSP(this, &FInstancedStructDetails::CanPasteInstancedStruct)))
 		.NameContent()
 		[
 			StructPropertyHandle->CreatePropertyNameWidget()
@@ -427,6 +431,68 @@ void FInstancedStructDetails::CustomizeHeader(TSharedRef<class IPropertyHandle>
 		.IsEnabled(StructProperty->IsEditable());
 }
 
+void FInstancedStructDetails::OnCopyInstancedStruct() const
+{
+	FString Value;
+	StructProperty->GetValueAsFormattedString(Value, PPF_Copy);
+
+	FPlatformApplicationMisc::ClipboardCopy(*Value);
+}
+
+bool FInstancedStructDetails::CanPasteInstancedStruct() const
+{
+	FString Value;
+	FPlatformApplicationMisc::ClipboardPaste(Value);
+
+	if (StructProperty && StructPicker && !Value.IsEmpty())
+	{
+		bool bCanPasteValue = true;
+		FInstancedStruct TempInstancedStructValue;
+
+		// This value actually a valid instanced struct?
+		{
+			const TCHAR* Buffer = *Value;
+			bCanPasteValue = TempInstancedStructValue.ImportTextItem(Buffer, PPF_Copy, nullptr, nullptr);
+		}
+
+		// Is this value compatible with our property type?
+		if (bCanPasteValue)
+		{
+			if (!StructPicker->CanChangeStructType())
+			{
+				// Can only apply the value if the new struct type matches the existing struct type of every instance
+				StructProperty->EnumerateConstRawData([&bCanPasteValue, &TempInstancedStructValue](const void* RawData, const int32 /*DataIndex*/, const int32 /*NumDatas*/)
+				{
+					const FInstancedStruct* InstancedStruct = static_cast<const FInstancedStruct*>(RawData);
+					bCanPasteValue = InstancedStruct && InstancedStruct->GetScriptStruct() == TempInstancedStructValue.GetScriptStruct();
+					return bCanPasteValue;
+				});
+			}
+			else if (const UScriptStruct* BaseScriptStruct = StructPicker->GetBaseScriptStruct())
+			{
+				// Can only apply the value if the new struct type is compatible with our base struct type, or the new struct type is null
+				if (const UScriptStruct* TempInstancedScriptStruct = TempInstancedStructValue.GetScriptStruct())
+				{
+					bCanPasteValue = TempInstancedScriptStruct->IsChildOf(BaseScriptStruct);
+				}
+			}
+		}
+
+		return bCanPasteValue;
+	}
+
+	return false;
+}
+
+void FInstancedStructDetails::OnPasteInstancedStruct()
+{
+	FString Value;
+	FPlatformApplicationMisc::ClipboardPaste(Value);
+
+	FScopedTransaction Transaction(LOCTEXT("Paste", "Paste Property"));
+	StructProperty->SetValueFromFormattedString(*Value, PPF_Copy);
+}
+
 void FInstancedStructDetails::OnObjectsReinstanced(const FReplacementObjectMap& ObjectMap)
 {
 	// Force update the details when BP is compiled, since we may cached hold references to the old object or class.
diff --git a/Engine/Source/Editor/StructUtilsEditor/Private/SInstancedStructPicker.cpp b/Engine/Source/Editor/StructUtilsEditor/Private/SInstancedStructPicker.cpp
--- a/Engine/Source/Editor/StructUtilsEditor/Private/SInstancedStructPicker.cpp
+++ b/Engine/Source/Editor/StructUtilsEditor/Private/SInstancedStructPicker.cpp
@@ -159,12 +159,14 @@ void SInstancedStructPicker::Construct(const FArguments& InArgs, TSharedPtr<IPro
 		}
 	}
 
+	bCanChangeStructType = !StructProperty->HasMetaData(NAME_StructTypeConst);
+
 	ChildSlot
 	[
 		SAssignNew(ComboButton, SComboButton)
 		.OnGetMenuContent(this, &SInstancedStructPicker::GenerateStructPicker)
 		.ContentPadding(0)
-		.IsEnabled(StructProperty->IsEditable() && !StructProperty->HasMetaData(NAME_StructTypeConst))
+		.IsEnabled(StructProperty->IsEditable() && bCanChangeStructType)
 		.ButtonContent()
 		[
 			SNew(SHorizontalBox)
diff --git a/Engine/Source/Editor/StructUtilsEditor/Public/InstancedStructDetails.h b/Engine/Source/Editor/StructUtilsEditor/Public/InstancedStructDetails.h
--- a/Engine/Source/Editor/StructUtilsEditor/Public/InstancedStructDetails.h
+++ b/Engine/Source/Editor/StructUtilsEditor/Public/InstancedStructDetails.h
@@ -33,6 +33,11 @@ public:
 
 private:
 
+	void OnCopyInstancedStruct() const;
+
+	bool CanPasteInstancedStruct() const;
+	void OnPasteInstancedStruct();
+
 	using FReplacementObjectMap = TMap<UObject*, UObject*>;
 	void OnObjectsReinstanced(const FReplacementObjectMap& ObjectMap);
 
diff --git a/Engine/Source/Editor/StructUtilsEditor/Public/SInstancedStructPicker.h b/Engine/Source/Editor/StructUtilsEditor/Public/SInstancedStructPicker.h
--- a/Engine/Source/Editor/StructUtilsEditor/Public/SInstancedStructPicker.h
+++ b/Engine/Source/Editor/StructUtilsEditor/Public/SInstancedStructPicker.h
@@ -50,6 +50,16 @@ public:
 
 	void Construct(const FArguments& InArgs, TSharedPtr<IPropertyHandle> InStructProperty, TSharedPtr<IPropertyUtilities> InPropertyUtils);
 
+	const UScriptStruct* GetBaseScriptStruct() const
+	{
+		return BaseScriptStruct.Get();
+	}
+
+	bool CanChangeStructType() const
+	{
+		return bCanChangeStructType;
+	}
+
 	FOnStructPicked OnStructPicked;
 
 private:
@@ -60,6 +70,9 @@ private:
 	/** The base struct that we're allowing to be picked (controlled by the "BaseStruct" meta-data) */
 	TWeakObjectPtr<UScriptStruct> BaseScriptStruct = nullptr;
 
+	/** Can we change the struct type at all? (controlled by the "StructTypeConst" meta-data) */
+	bool bCanChangeStructType = true;
+
 	/** The array of allowed structs (additionally to the meta-data) */
 	TArray<TSoftObjectPtr<const UScriptStruct>> AllowedStructs;
 
diff --git a/Engine/Source/Editor/StructUtilsEditor/StructUtilsEditor.Build.cs b/Engine/Source/Editor/StructUtilsEditor/StructUtilsEditor.Build.cs
--- a/Engine/Source/Editor/StructUtilsEditor/StructUtilsEditor.Build.cs
+++ b/Engine/Source/Editor/StructUtilsEditor/StructUtilsEditor.Build.cs
@@ -32,6 +32,7 @@ namespace UnrealBuildTool.Rules
 
 			PrivateDependencyModuleNames.AddRange(
 				new string[] {
+					"ApplicationCore",
 					"RenderCore",
 					"GraphEditor",
 					"KismetWidgets",
diff --git a/Engine/Source/Runtime/CoreUObject/Private/StructUtils/InstancedStruct.cpp b/Engine/Source/Runtime/CoreUObject/Private/StructUtils/InstancedStruct.cpp
--- a/Engine/Source/Runtime/CoreUObject/Private/StructUtils/InstancedStruct.cpp
+++ b/Engine/Source/Runtime/CoreUObject/Private/StructUtils/InstancedStruct.cpp
@@ -1,5 +1,8 @@
 // Copyright Epic Games, Inc. All Rights Reserved.
+
 #include "StructUtils/InstancedStruct.h"
+
+#include "Misc/PackageName.h"
 #include "StructUtils/StructView.h"
 #include "Serialization/CustomVersion.h"
 #include "UObject/Linker.h"
@@ -284,6 +287,12 @@ bool FInstancedStruct::ImportTextItem(const TCHAR*& Buffer, int32 PortFlags, UOb
 			StructPathName.Append(NewName.ToString());
 		}
 
+		// Extra validation when pasting, as we may be importing random data from the clipboard and this avoids LoadObject warning log spam
+		if ((PortFlags & PPF_Copy) && !FPackageName::IsValidObjectPath(StructPathName.ToView()))
+		{
+			return false;
+		}
+
 		// Make sure the struct is actually loaded before trying to import the text (this boils down to FindObject if the struct is already loaded).
 		// This is needed for user defined structs, BP pin values, config, copy/paste, where there's no guarantee that the referenced struct has actually been loaded yet.
 		UScriptStruct* StructTypePtr = LoadObject<UScriptStruct>(nullptr, StructPathName.ToString());
