[StructUtils] Prevent renaming when the new name is the same as the previous name.

diff --git a/Engine/Source/Editor/StructUtilsEditor/Private/Customizations/PropertyBagDetails.cpp b/Engine/Source/Editor/StructUtilsEditor/Private/Customizations/PropertyBagDetails.cpp
--- a/Engine/Source/Editor/StructUtilsEditor/Private/Customizations/PropertyBagDetails.cpp
+++ b/Engine/Source/Editor/StructUtilsEditor/Private/Customizations/PropertyBagDetails.cpp
@@ -1462,6 +1462,10 @@ void FPropertyBagInstanceDataDetails::OnChildRowAdded(IDetailPropertyRow& ChildR
 
 				const FName NewName = FInstancedPropertyBag::SanitizePropertyName(InNewText.ToString());
 				const FName OldName = UE::StructUtils::Private::GetPropertyNameSafe(ChildPropertyHandle);
+				if (NewName == OldName)
+				{
+					return;
+				}
 
 				if (!ensureMsgf(UE::StructUtils::Private::IsUniqueName(NewName, OldName, BagStructProperty), TEXT("Should have already been addressed in OnVerifyTextChanged.")))
 				{
