[StructUtils] Property Bag inline rename changed to display the actual value

diff --git a/Engine/Source/Editor/StructUtilsEditor/Private/Customizations/PropertyBagDetails.cpp b/Engine/Source/Editor/StructUtilsEditor/Private/Customizations/PropertyBagDetails.cpp
--- a/Engine/Source/Editor/StructUtilsEditor/Private/Customizations/PropertyBagDetails.cpp
+++ b/Engine/Source/Editor/StructUtilsEditor/Private/Customizations/PropertyBagDetails.cpp
@@ -1409,9 +1409,8 @@ void FPropertyBagInstanceDataDetails::OnChildRowAdded(IDetailPropertyRow& ChildR
 			.OverflowPolicy(ETextOverflowPolicy::Ellipsis)
 			.Text_Lambda([ChildPropertyHandle]()
 			{
-				const bool bIsBool = ChildPropertyHandle && ChildPropertyHandle->GetPropertyClass() == FBoolProperty::StaticClass();
 				const FName PropertyName = UE::StructUtils::Private::GetPropertyNameSafe(ChildPropertyHandle);
-				return FText::FromString(FName::NameToDisplayString(PropertyName.ToString(), bIsBool));
+				return FText::FromName(PropertyName);
 			})
 			.OnVerifyTextChanged_Lambda([BagStructProperty = BagStructProperty, ChildPropertyHandle](const FText& InText, FText& OutErrorMessage)
 			{
@@ -1462,6 +1461,7 @@ void FPropertyBagInstanceDataDetails::OnChildRowAdded(IDetailPropertyRow& ChildR
 
 				const FName NewName = FInstancedPropertyBag::SanitizePropertyName(InNewText.ToString());
 				const FName OldName = UE::StructUtils::Private::GetPropertyNameSafe(ChildPropertyHandle);
+
 				if (NewName == OldName)
 				{
 					return;
