[StructUtils] Added Metadata Options to the StructUtils PropertyBag DetailsView

Note: The UI layout for these options is still in flux and will likely go through a feedback loop

#rnx

diff --git a/Engine/Source/Editor/StructUtilsEditor/Private/Customizations/PropertyBagDetails.cpp b/Engine/Source/Editor/StructUtilsEditor/Private/Customizations/PropertyBagDetails.cpp
--- a/Engine/Source/Editor/StructUtilsEditor/Private/Customizations/PropertyBagDetails.cpp
+++ b/Engine/Source/Editor/StructUtilsEditor/Private/Customizations/PropertyBagDetails.cpp
@@ -20,6 +20,8 @@
 #include "Widgets/Images/SImage.h"
 #include "Widgets/Input/SButton.h"
 #include "Widgets/Input/SComboButton.h"
+#include "Widgets/Input/SEditableTextBox.h"
+#include "Widgets/Input/SMultiLineEditableTextBox.h"
 #include "Widgets/Text/SInlineEditableTextBlock.h"
 
 #include UE_INLINE_GENERATED_CPP_BY_NAME(PropertyBagDetails)
@@ -996,6 +998,10 @@ TSharedRef<SWidget> CreateTypeSelectionWidget(TSharedPtr<IPropertyHandle> ChildP
 namespace Constants
 {
 	static constexpr int32 MaxCategoryLength = 70;
+	static constexpr int32 MinSubMenuTextBoxWidth = 40;
+	static constexpr int32 MaxSubMenuTextBoxWidth = 800;
+	static constexpr int32 SubMenuTextBoxPadding = 5;
+
 	// Special case for categories. Alphanumeric, but including spaces and `|` for nested categories.
 	static constexpr TCHAR InvalidCategoryCharacters[] = TEXT("\"',/.:&!?~\\\n\r\t@#(){}[]<>=;^%$`*+-");
 }
@@ -1274,6 +1280,7 @@ void FPropertyBagInstanceDataDetails::OnGroupRowAdded(IDetailGroup& GroupRow, in
 
 void FPropertyBagInstanceDataDetails::OnChildRowAdded(IDetailPropertyRow& ChildRow)
 {
+	using namespace UE::StructUtils;
 	TSharedPtr<SWidget> NameWidget;
 	TSharedPtr<SWidget> PropertyValueWidget;
 	FDetailWidgetRow DetailWidgetRow;
@@ -1284,7 +1291,7 @@ void FPropertyBagInstanceDataDetails::OnChildRowAdded(IDetailPropertyRow& ChildR
 
 	TWeakPtr<FPropertyBagInstanceDataDetails> WeakSelf = SharedThis<FPropertyBagInstanceDataDetails>(this);
 
-	const UPropertyBag* BagStruct = BagStructProperty ? UE::StructUtils::Private::GetCommonBagStruct(BagStructProperty) : nullptr;
+	const UPropertyBag* BagStruct = BagStructProperty ? Private::GetCommonBagStruct(BagStructProperty) : nullptr;
 	const FProperty* ChildProperty = ChildPropertyHandle->GetProperty();
 	const FPropertyBagPropertyDesc* PropertyDesc = BagStruct ? BagStruct->FindPropertyDescByProperty(ChildProperty) : nullptr;
 
@@ -1293,7 +1300,7 @@ void FPropertyBagInstanceDataDetails::OnChildRowAdded(IDetailPropertyRow& ChildR
 	PRAGMA_ENABLE_DEPRECATION_WARNINGS
 
 	// Validate data and check if it's editable
-	if (ChildProperty->HasMetaData(UE::StructUtils::Metadata::HideInDetailPanelsName))
+	if (!ChildProperty || ChildProperty->HasMetaData(Metadata::HideInDetailPanelsName))
 	{
 		ChildRow.Visibility(EVisibility::Collapsed);
 		return;
@@ -1313,11 +1320,11 @@ void FPropertyBagInstanceDataDetails::OnChildRowAdded(IDetailPropertyRow& ChildR
 		// Warn that the unsigned types cannot be set via the type selection.
 		WarningOnProperty = LOCTEXT("UnsignedTypesWarning", "Unsigned types are not supported through the property type selection. If you change the type, you will not be able to change it back.");
 	}
-	else if (UE::StructUtils::Private::HasMissingType(ChildPropertyHandle))
+	else if (Private::HasMissingType(ChildPropertyHandle))
 	{
 		WarningOnProperty = LOCTEXT("MissingTypeWarning", "The property is missing type. The Struct, Enum, or Object may have been removed.");
 	}
-	else if (!FInstancedPropertyBag::IsPropertyNameValid(UE::StructUtils::Private::GetPropertyNameSafe(ChildPropertyHandle)))
+	else if (!FInstancedPropertyBag::IsPropertyNameValid(Private::GetPropertyNameSafe(ChildPropertyHandle)))
 	{
 		WarningOnProperty = LOCTEXT("InvalidNameWarning", "The property's name contains invalid characters. Dynamically named properties with invalid characters may be rejected in future releases.");
 	}
@@ -1397,19 +1404,25 @@ void FPropertyBagInstanceDataDetails::OnChildRowAdded(IDetailPropertyRow& ChildR
 				.Padding(1, 0)
 				.AutoWidth()
 				[
-					UE::StructUtils::CreateTypeSelectionWidget(ChildPropertyHandle, BagStructProperty, PropUtils, SPinTypeSelector::ESelectorType::Compact, bAllowContainers)
+					CreateTypeSelectionWidget(ChildPropertyHandle, BagStructProperty, PropUtils, SPinTypeSelector::ESelectorType::Compact, bAllowContainers)
 				];
 		}
 
 		/*** EDITABLE NAME BLOCK ***/
 		TSharedPtr<SInlineEditableTextBlock> EditableInlineNameWidget = SNew(SInlineEditableTextBlock)
 			.IsReadOnly(!EnumHasAnyFlags(ChildRowFeatures, EPropertyBagChildRowFeatures::Renaming))
+			.ToolTipText_Lambda([ChildPropertyHandle]() -> FText
+			{
+				return ChildPropertyHandle->HasMetaData(Metadata::Specifiers::ToolTipName)
+						   ? FText::FromString(ChildPropertyHandle->GetMetaData(Metadata::Specifiers::ToolTipName))
+						   : FText::GetEmpty();
+			})
 			.Font(IDetailLayoutBuilder::GetDetailFont())
 			.MultiLine(false)
 			.OverflowPolicy(ETextOverflowPolicy::Ellipsis)
 			.Text_Lambda([ChildPropertyHandle]()
 			{
-				const FName PropertyName = UE::StructUtils::Private::GetPropertyNameSafe(ChildPropertyHandle);
+				const FName PropertyName = Private::GetPropertyNameSafe(ChildPropertyHandle);
 				return FText::FromName(PropertyName);
 			})
 			.OnVerifyTextChanged_Lambda([BagStructProperty = BagStructProperty, ChildPropertyHandle](const FText& InText, FText& OutErrorMessage)
@@ -1427,7 +1440,7 @@ void FPropertyBagInstanceDataDetails::OnChildRowAdded(IDetailPropertyRow& ChildR
 					return false;
 				}
 
-				const FName OldName = UE::StructUtils::Private::GetPropertyNameSafe(ChildPropertyHandle);
+				const FName OldName = Private::GetPropertyNameSafe(ChildPropertyHandle);
 				// Bypass if the name is the exact same.
 				if (InText.ToString().Equals(OldName.ToString()))
 				{
@@ -1443,7 +1456,7 @@ void FPropertyBagInstanceDataDetails::OnChildRowAdded(IDetailPropertyRow& ChildR
 					return true;
 				}
 
-				if (!UE::StructUtils::Private::IsUniqueName(NewName, OldName, BagStructProperty))
+				if (!Private::IsUniqueName(NewName, OldName, BagStructProperty))
 				{
 					OutErrorMessage = LOCTEXT("InlinePropertyUniqueName", "Property must have unique name");
 					return false;
@@ -1460,19 +1473,19 @@ void FPropertyBagInstanceDataDetails::OnChildRowAdded(IDetailPropertyRow& ChildR
 				}
 
 				const FName NewName = FInstancedPropertyBag::SanitizePropertyName(InNewText.ToString());
-				const FName OldName = UE::StructUtils::Private::GetPropertyNameSafe(ChildPropertyHandle);
+				const FName OldName = Private::GetPropertyNameSafe(ChildPropertyHandle);
 
 				if (NewName == OldName)
 				{
 					return;
 				}
 
-				if (!ensureMsgf(UE::StructUtils::Private::IsUniqueName(NewName, OldName, BagStructProperty), TEXT("Should have already been addressed in OnVerifyTextChanged.")))
+				if (!ensureMsgf(Private::IsUniqueName(NewName, OldName, BagStructProperty), TEXT("Should have already been addressed in OnVerifyTextChanged.")))
 				{
 					return;
 				}
 
-				UE::StructUtils::Private::ApplyChangesToPropertyDescs(
+				Private::ApplyChangesToPropertyDescs(
 					FText::Format(LOCTEXT("OnPropertyNameChanged", "Change property name: {0} -> {1}"), FText::FromName(OldName), FText::FromName(NewName)),
 					BagStructProperty,
 					PropUtils,
@@ -1552,7 +1565,7 @@ void FPropertyBagInstanceDataDetails::OnChildRowAdded(IDetailPropertyRow& ChildR
 					.OnClicked_Lambda([ChildPropertyHandle, BagStructProperty = BagStructProperty, PropUtils = PropUtils]()
 					{
 						FProperty* Property = ChildPropertyHandle->GetProperty();
-						UE::StructUtils::Private::ApplyChangesToPropertyDescs(
+						Private::ApplyChangesToPropertyDescs(
 							LOCTEXT("OnPropertyAccessSpecifierChanged", "Set access specifier."),
 							BagStructProperty,
 							PropUtils,
@@ -1613,7 +1626,7 @@ void FPropertyBagInstanceDataDetails::OnChildRowAdded(IDetailPropertyRow& ChildR
 			{
 				MenuBuilder.BeginSection(/*InExtensionHook=*/NAME_None, LOCTEXT("DropDownMenuSectionTypeSelector", "Type"));
 				MenuBuilder.AddWidget(
-					UE::StructUtils::CreateTypeSelectionWidget(
+					CreateTypeSelectionWidget(
 						ChildPropertyHandle,
 						BagStructProperty,
 						PropUtils,
@@ -1637,9 +1650,9 @@ void FPropertyBagInstanceDataDetails::OnChildRowAdded(IDetailPropertyRow& ChildR
 						LOCTEXT("DropDownMenuRenameProperty", "Rename property"),
 						LOCTEXT("DropDownMenuRenamePropertyToolTip", "Enable the inline property renaming."),
 						FSlateIcon(FAppStyle::GetAppStyleSetName(), "Icons.Edit"),
-						FUIAction(FExecuteAction::CreateLambda([NameWidget = EditableInlineNameWidget]()
+						FUIAction(FExecuteAction::CreateLambda([NameWidgetWeak = EditableInlineNameWidget.ToWeakPtr()]()
 						{
-							if (NameWidget)
+							if (TSharedPtr<SInlineEditableTextBlock> NameWidget = NameWidgetWeak.Pin())
 							{
 								NameWidget->EnterEditingMode();
 							}
@@ -1654,7 +1667,7 @@ void FPropertyBagInstanceDataDetails::OnChildRowAdded(IDetailPropertyRow& ChildR
 						FSlateIcon(FAppStyle::GetAppStyleSetName(), "Icons.Delete"),
 						FUIAction(FExecuteAction::CreateLambda([BagStructProperty = BagStructProperty, PropUtils = PropUtils, ChildPropertyHandle]()
 						{
-							UE::StructUtils::Private::DeleteProperty(BagStructProperty, ChildPropertyHandle, PropUtils);
+							Private::DeleteProperty(BagStructProperty, ChildPropertyHandle, PropUtils);
 						}))
 					);
 				}
@@ -1662,12 +1675,44 @@ void FPropertyBagInstanceDataDetails::OnChildRowAdded(IDetailPropertyRow& ChildR
 				MenuBuilder.EndSection();
 			}
 
+			// BagStructProperty and PropUtils are forwarded straight to OnTextCommitted, so they save a ref count and be refs here.
+			auto BuildMetadataSpecifierSubMenu = [](
+				TAttribute<FText>&& TextAttribute,
+				FOnVerifyTextChanged&& OnVerifyTextDelegate,
+				FOnTextCommitted&& OnTextCommittedDelegate) -> FNewMenuDelegate
+			{
+				TSharedRef<SWidget> TextBox = SNew(SMultiLineEditableTextBox)
+					.WrapTextAt(400)
+					.Text(MoveTemp(TextAttribute))
+					.OnVerifyTextChanged(MoveTemp(OnVerifyTextDelegate))
+					.OnTextCommitted(OnTextCommittedDelegate);
+
+				return FNewMenuDelegate::CreateLambda(
+					[TextBox](FMenuBuilder& OutMenuBuilder)
+					{
+						OutMenuBuilder.AddWidget(
+							SNew(SHorizontalBox)
+							+ SHorizontalBox::Slot()
+							.HAlign(HAlign_Fill)
+							.MinWidth(Constants::MinSubMenuTextBoxWidth)
+							.MaxWidth(Constants::MaxSubMenuTextBoxWidth)
+							.Padding(Constants::SubMenuTextBoxPadding, 0)
+							.AutoWidth()
+							[
+								TextBox
+							],
+							FText::GetEmpty(),
+							/*bNoIndent=*/true,
+							/*bInSearchable=*/true);
+					});
+			};
+
 			// The property's category (grouping) can be edited here.
 			if (EnumHasAllFlags(ChildRowFeatures, EPropertyBagChildRowFeatures::Categories | EPropertyBagChildRowFeatures::Menu_Categories))
 			{
 				MenuBuilder.BeginSection(/*InExtensionHook=*/NAME_None, LOCTEXT("DropDownMenuSectionCategory", "Category"));
 
-				if (ChildPropertyHandle && ChildPropertyHandle->HasMetaData(UE::StructUtils::Metadata::CategoryName))
+				if (ChildPropertyHandle->HasMetaData(Metadata::CategoryName))
 				{
 					MenuBuilder.AddMenuEntry(
 						LOCTEXT("DropDownMenuClearCategory", "Clear category"),
@@ -1675,64 +1720,166 @@ void FPropertyBagInstanceDataDetails::OnChildRowAdded(IDetailPropertyRow& ChildR
 						FSlateIcon(FAppStyle::GetAppStyleSetName(), "Icons.Delete"),
 						FUIAction(FExecuteAction::CreateLambda([ChildPropertyHandle, StructProperty = BagStructProperty, PropUtils = PropUtils]()
 						{
-							UE::StructUtils::Private::ApplyChangesToSinglePropertyDesc(
+							Private::ApplyChangesToSinglePropertyDesc(
 								LOCTEXT("DropDownMenuOnCategoryCleared", "Clear property category"),
 								ChildPropertyHandle,
 								StructProperty,
 								PropUtils,
 								[](FPropertyBagPropertyDesc& Desc)
 								{
-									UE::StructUtils::Metadata::RemoveCategory(Desc);
+									Metadata::RemoveCategory(Desc);
 								});
 						})));
 				}
 
-				// TODO: AddVerifiedEditableText seems to bypass the MenuBuilder Section, so this was added temporarily to force the section to exist.
-				MenuBuilder.AddWidget(SNullWidget::NullWidget, FText::GetEmpty(), false, false);
-
-				MenuBuilder.AddVerifiedEditableText(
-					LOCTEXT("DropDownMenuSubMenuCategoryName", "Category"),
-					LOCTEXT("DropDownMenuSubMenuCategoryTooltip", "Edit this value to change the category of this property. Subcategories can be created with the '|' character."),
-					FSlateIcon(FAppStyle::GetAppStyleSetName(), "Icons.FolderOpen"),
-					TAttribute<FText>::CreateLambda([ChildPropertyHandle]()
+				FOnVerifyTextChanged ValidateCategoryNameDelegate = FOnVerifyTextChanged::CreateLambda([](const FText& InText, FText& OutMessage)
+				{
+					if (InText.ToString().Len() > Constants::MaxCategoryLength)
+					{
+						OutMessage = LOCTEXT("DropDownMenuInvalidCategoryName", "Invalid category name");
+						return false;
+					}
+					else
 					{
-						FText GroupLabel = FText::FromString("");
-						if (ChildPropertyHandle && ChildPropertyHandle->HasMetaData(UE::StructUtils::Metadata::CategoryName))
+						return true;
+					}
+				});
+
+				MenuBuilder.AddSubMenu(
+					FText::FromName(Metadata::CategoryName),
+					LOCTEXT("DropDownMenuSubMenuCategoryTooltip", "Set the category of this property. Subcategories can be created with the '|' character."),
+					BuildMetadataSpecifierSubMenu(
+						TAttribute<FText>::CreateLambda([ChildPropertyHandle]()
 						{
-							GroupLabel = FText::FromString(ChildPropertyHandle->GetMetaData(UE::StructUtils::Metadata::CategoryName));
-						}
+							FText GroupLabel = FText::FromString("");
+							if (ChildPropertyHandle && ChildPropertyHandle->HasMetaData(Metadata::CategoryName))
+							{
+								GroupLabel = FText::FromString(ChildPropertyHandle->GetMetaData(Metadata::CategoryName));
+							}
 
-						return GroupLabel;
-					}),
-					FOnVerifyTextChanged::CreateLambda([](const FText& InText, FText& OutMessage)
-					{
-						if (InText.ToString().Len() > UE::StructUtils::Constants::MaxCategoryLength)
+							return GroupLabel;
+						}),
+						FOnVerifyTextChanged::CreateLambda([](const FText& InText, FText& OutMessage)
 						{
-							OutMessage = LOCTEXT("DropDownMenuInvalidCategoryName", "Invalid category name");
-							return false;
-						}
-						else
+							if (InText.ToString().Len() > Constants::MaxCategoryLength)
+							{
+								OutMessage = LOCTEXT("DropDownMenuInvalidCategoryName", "Invalid category name");
+								return false;
+							}
+							else
+							{
+								return true;
+							}
+						}),
+
+						FOnTextCommitted::CreateLambda([StructProperty = BagStructProperty, PropUtils = PropUtils, ChildPropertyHandle](const FText& CommittedText, const ETextCommit::Type CommitType)
 						{
-							return true;
-						}
-					}),
-					FOnTextCommitted::CreateLambda([StructProperty = BagStructProperty, PropUtils = PropUtils, ChildPropertyHandle](const FText& CommittedText, const ETextCommit::Type CommitType)
-					{
-						if (CommitType == ETextCommit::OnEnter || CommitType == ETextCommit::OnUserMovedFocus)
+							if (CommitType == ETextCommit::OnEnter || CommitType == ETextCommit::OnUserMovedFocus)
+							{
+								Private::ApplyChangesToSinglePropertyDesc(
+									LOCTEXT("DropDownMenuOnCategoryEdited", "Edit property category"),
+									ChildPropertyHandle,
+									StructProperty,
+									PropUtils,
+									[&CommittedText](FPropertyBagPropertyDesc& Desc)
+									{
+										Metadata::SetCategory(Desc, CommittedText.ToString());
+									});
+							}
+						})),
+					/*bInOpenSubMenuOnClick=*/true);
+
+				MenuBuilder.EndSection();
+			}
+
+			// The property's optional metadata specifiers can be edited here.
+			if (EnumHasAllFlags(ChildRowFeatures, EPropertyBagChildRowFeatures::Menu_MetadataSpecifiers))
+			{
+				MenuBuilder.BeginSection(/*InExtensionHook=*/NAME_None, LOCTEXT("DropDownMenuSectionMetadataSpecifiers", "Specifiers"));
+
+				auto BuildGenericSpecifierSubMenu = [this, ChildPropertyHandle, BuildMetadataSpecifierSubMenu](const FName& SpecifierName, const bool bIsNumeric = false)
+				{
+					return BuildMetadataSpecifierSubMenu(
+						TAttribute<FText>::CreateLambda([SpecifierName, ChildPropertyHandle]()
 						{
-							UE::StructUtils::Private::ApplyChangesToSinglePropertyDesc(
-								LOCTEXT("DropDownMenuOnCategoryEdited", "Edit property category"),
-								ChildPropertyHandle,
-								StructProperty,
-								PropUtils,
-								[&CommittedText](FPropertyBagPropertyDesc& Desc)
+							return ChildPropertyHandle.IsValid()
+									   ? FText::FromString(ChildPropertyHandle->GetMetaData(SpecifierName))
+									   : FText::GetEmpty();
+						}),
+						!bIsNumeric
+							? FOnVerifyTextChanged()
+							: FOnVerifyTextChanged::CreateLambda([](const FText& InText, FText& OutMessage) -> bool
+							{
+								if (!InText.IsNumeric())
 								{
-									UE::StructUtils::Metadata::SetCategory(Desc, CommittedText.ToString());
-								});
-						}
-					}),
-					FOnTextChanged(),
-					!bEditable);
+									OutMessage = LOCTEXT("DropDownMenuSpecifierNumericValueExpected", "Value is not numeric");
+									return false;
+								}
+								else
+								{
+									return true;
+								}
+							}),
+						FOnTextCommitted::CreateLambda(
+							[SpecifierName, StructProperty = BagStructProperty, PropUtils = PropUtils, ChildPropertyHandle]
+							(const FText& CommittedText, const ETextCommit::Type CommitType)
+							{
+								if (CommitType != ETextCommit::OnCleared)
+								{
+									Private::ApplyChangesToSinglePropertyDesc(
+										FText::Format(LOCTEXT("DropDownMenuOnSetSpecifier", "Set metadata specifier '{0}' value"), FText::FromName(SpecifierName)),
+										ChildPropertyHandle,
+										StructProperty,
+										PropUtils,
+										[SpecifierName, &CommittedText](FPropertyBagPropertyDesc& Desc)
+										{
+											Desc.SetMetaData(SpecifierName, CommittedText.ToString());
+										});
+								}
+							}));
+				};
+
+				MenuBuilder.AddSubMenu(
+					FText::FromName(Metadata::Specifiers::ToolTipName),
+					LOCTEXT("DropDownMenuSubMenuTooltipTooltip", "Set the tooltip for this property."),
+					BuildGenericSpecifierSubMenu(Metadata::Specifiers::ToolTipName),
+					/*bInOpenSubMenuOnClick=*/true);
+
+				if (Metadata::Specifiers::SupportsClamping(*PropertyDesc))
+				{
+					MenuBuilder.AddSubMenu(
+						FText::FromName(Metadata::Specifiers::ClampMinName),
+						LOCTEXT("DropDownMenuSubMenuClampMinTooltip", "Set the min value for this property."),
+						BuildGenericSpecifierSubMenu(Metadata::Specifiers::ClampMinName, /*bIsNumeric=*/true),
+						/*bInOpenSubMenuOnClick=*/true);
+
+					MenuBuilder.AddSubMenu(
+						FText::FromName(Metadata::Specifiers::ClampMaxName),
+						LOCTEXT("DropDownMenuSubMenuClampMaxTooltip", "Set the max value for this property."),
+						BuildGenericSpecifierSubMenu(Metadata::Specifiers::ClampMaxName, /*bIsNumeric=*/true),
+						/*bInOpenSubMenuOnClick=*/true);
+
+					MenuBuilder.AddSubMenu(
+						FText::FromName(Metadata::Specifiers::UIMinName),
+						LOCTEXT("DropDownMenuSubMenuUIMinTooltip", "Set the min UI value for this property."),
+						BuildGenericSpecifierSubMenu(Metadata::Specifiers::UIMinName, /*bIsNumeric=*/true),
+						/*bInOpenSubMenuOnClick=*/true);
+
+					MenuBuilder.AddSubMenu(
+						FText::FromName(Metadata::Specifiers::UIMaxName),
+						LOCTEXT("DropDownMenuSubMenuUIMaxTooltip", "Set the max UI value for this property."),
+						BuildGenericSpecifierSubMenu(Metadata::Specifiers::UIMaxName, /*bIsNumeric=*/true),
+						/*bInOpenSubMenuOnClick=*/true);
+				}
+
+				if (Metadata::Specifiers::SupportsSettingUnits(*PropertyDesc))
+				{
+					MenuBuilder.AddSubMenu(
+						FText::FromName(Metadata::Specifiers::UnitsName),
+						LOCTEXT("DropDownMenuSubMenuUnitsTooltip", "Set the units descriptor for this property."),
+						BuildGenericSpecifierSubMenu(Metadata::Specifiers::UnitsName),
+						/*bInOpenSubMenuOnClick=*/true);
+				}
 
 				MenuBuilder.EndSection();
 			}
@@ -2179,7 +2326,6 @@ void FPropertyBagDetails::CustomizeHeader(TSharedRef<IPropertyHandle> StructProp
 		}
 
 		// Load the feature set by the metadata set on the FPropertyBag. Can only accept explicit enum values currently.
-		// TODO: Enable the option to parse a bitflag expression from string. I.E. 'Renaming | DropDownMenuButton | AllMenuOptions'
 		if (MetaDataProperty->HasMetaData(UE::StructUtils::Metadata::ChildRowFeaturesName))
 		{
 			if (const UEnum* Enum = StaticEnum<EPropertyBagChildRowFeatures>())
diff --git a/Engine/Source/Editor/StructUtilsEditor/Private/StructUtilsMetadata.cpp b/Engine/Source/Editor/StructUtilsEditor/Private/StructUtilsMetadata.cpp
--- a/Engine/Source/Editor/StructUtilsEditor/Private/StructUtilsMetadata.cpp
+++ b/Engine/Source/Editor/StructUtilsEditor/Private/StructUtilsMetadata.cpp
@@ -22,9 +22,9 @@ namespace UE::StructUtils::Metadata
 	const FLazyName CanRemovePropertyName = "CanRemoveProperty";
 	const FLazyName ChildRowFeaturesName = "ChildRowFeatures";
 
-	bool AreCategoriesEnabled(const FPropertyBagPropertyDesc& OutDesc)
+	bool AreCategoriesEnabled(const FPropertyBagPropertyDesc& Desc)
 	{
-		return OutDesc.MetaData.ContainsByPredicate([](const FPropertyBagPropertyDescMetaData& MetaData) { return MetaData.Key == EnableCategoriesName; });
+		return Desc.HasMetaData(EnableCategoriesName);
 	}
 
 	void EnableCategories(FPropertyBagPropertyDesc& OutDesc)
@@ -45,14 +45,7 @@ namespace UE::StructUtils::Metadata
 
 	void SetCategory(FPropertyBagPropertyDesc& OutDesc, const FString& InGroupLabel, const bool bAutoEnableCategories)
 	{
-		if (FPropertyBagPropertyDescMetaData* MetaData = OutDesc.MetaData.FindByPredicate([](const FPropertyBagPropertyDescMetaData& InMetaData) { return InMetaData.Key == CategoryName; }))
-		{
-			MetaData->Value = InGroupLabel;
-		}
-		else
-		{
-			OutDesc.MetaData.Emplace(FPropertyBagPropertyDescMetaData(CategoryName, InGroupLabel));
-		}
+		OutDesc.SetMetaData(CategoryName, InGroupLabel);
 
 		if (bAutoEnableCategories)
 		{
@@ -75,11 +68,26 @@ namespace UE::StructUtils::Metadata
 
 	FString GetCategory(const FPropertyBagPropertyDesc& InDesc)
 	{
-		if (const FPropertyBagPropertyDescMetaData* MetaData = InDesc.MetaData.FindByPredicate([](const FPropertyBagPropertyDescMetaData& InMetaData) { return InMetaData.Key == CategoryName; }))
+		return InDesc.GetMetaData(CategoryName);
+	}
+
+	namespace Specifiers
+	{
+		const FLazyName ToolTipName = "ToolTip";
+		const FLazyName ClampMinName = "ClampMin";
+		const FLazyName ClampMaxName = "ClampMax";
+		const FLazyName UIMinName = "UIMin";
+		const FLazyName UIMaxName = "UIMax";
+		const FLazyName UnitsName = "Units";
+
+		bool SupportsClamping(const FPropertyBagPropertyDesc& InDesc)
 		{
-			return MetaData->Value;
+			return InDesc.IsNumericType();
 		}
 
-		return FString("");
+		bool SupportsSettingUnits(const FPropertyBagPropertyDesc& InDesc)
+		{
+			return InDesc.IsNumericType();
+		}
 	}
 }
diff --git a/Engine/Source/Editor/StructUtilsEditor/Public/PropertyBagDetails.h b/Engine/Source/Editor/StructUtilsEditor/Public/PropertyBagDetails.h
--- a/Engine/Source/Editor/StructUtilsEditor/Public/PropertyBagDetails.h
+++ b/Engine/Source/Editor/StructUtilsEditor/Public/PropertyBagDetails.h
@@ -43,12 +43,13 @@ enum class EPropertyBagChildRowFeatures : uint64
 	// Insert new drop general features above this line.
 
 	/** Menu Options (for drop-down menu or other future menus (right click) */
-	Menu_TypeSelector = 1ULL << 17, // Adds the type selector pill widget to the drop-down.
-	Menu_Rename       = 1ULL << 18, // Renaming the property from the drop-down menu. Requires property renaming enabled.
-	Menu_Delete       = 1ULL << 19, // Deleting the property from the drop-down menu.
-	Menu_Categories   = 1ULL << 20, // Create new/remove from categories
+	Menu_TypeSelector       = 1ULL << 17, // Adds the type selector pill widget to the drop-down.
+	Menu_Rename             = 1ULL << 18, // Renaming the property from the drop-down menu. Requires property renaming enabled.
+	Menu_Delete             = 1ULL << 19, // Deleting the property from the drop-down menu.
+	Menu_Categories         = 1ULL << 20, // Create new/remove from categories.
+	Menu_MetadataSpecifiers = 1ULL << 21, // Common metadata specifier options.
 
-	AllMenuOptions = Menu_TypeSelector | Menu_Rename | Menu_Delete | Menu_Categories,
+	AllMenuOptions = Menu_TypeSelector | Menu_Rename | Menu_Delete | Menu_Categories | Menu_MetadataSpecifiers,
 
 	// Insert new menu features above this line.
 	Deprecated = 1ULL << 63, // To allow for deprecating older features.
@@ -58,7 +59,7 @@ enum class EPropertyBagChildRowFeatures : uint64
 	// Renaming and deleting enabled, with type selection happening in the drop-down menu.
 	Core     = Renaming | Deletion | DropDownMenuButton | Menu_TypeSelector | Menu_Rename | Menu_Delete,
 	// Also enables the compact type selector icon, drag and drop, and categories support.
-	Extended = Core | DragAndDrop | CompactTypeSelector | Categories | Menu_Categories,
+	Extended = Core | DragAndDrop | CompactTypeSelector | Categories | Menu_Categories | Menu_MetadataSpecifiers,
 	// All options.
 	All      = AllGeneralOptions | AllMenuOptions,
 	// The default version includes deprecated UI features to support previous behavior.
diff --git a/Engine/Source/Editor/StructUtilsEditor/Public/StructUtilsMetadata.h b/Engine/Source/Editor/StructUtilsEditor/Public/StructUtilsMetadata.h
--- a/Engine/Source/Editor/StructUtilsEditor/Public/StructUtilsMetadata.h
+++ b/Engine/Source/Editor/StructUtilsEditor/Public/StructUtilsMetadata.h
@@ -43,15 +43,28 @@ namespace UE::StructUtils::Metadata
 	extern STRUCTUTILSEDITOR_API const FLazyName CategoryName;
 
 	/** If organizing by category is enabled for this property. */
-	extern STRUCTUTILSEDITOR_API bool AreCategoriesEnabled(const FPropertyBagPropertyDesc& OutDesc);
+	STRUCTUTILSEDITOR_API bool AreCategoriesEnabled(const FPropertyBagPropertyDesc& Desc);
 	/** Enable category organization for this property. */
-	extern STRUCTUTILSEDITOR_API void EnableCategories(FPropertyBagPropertyDesc& OutDesc);
+	STRUCTUTILSEDITOR_API void EnableCategories(FPropertyBagPropertyDesc& OutDesc);
 	/** Disable category organization for this property. */
-	extern STRUCTUTILSEDITOR_API void DisableCategories(FPropertyBagPropertyDesc& OutDesc);
+	STRUCTUTILSEDITOR_API void DisableCategories(FPropertyBagPropertyDesc& OutDesc);
 	/** Set the category for this property. */
-	extern STRUCTUTILSEDITOR_API void SetCategory(FPropertyBagPropertyDesc& OutDesc, const FString& InGroupLabel, bool bAutoEnableCategories = true);
+	STRUCTUTILSEDITOR_API void SetCategory(FPropertyBagPropertyDesc& OutDesc, const FString& InGroupLabel, bool bAutoEnableCategories = true);
 	/** Remove the category for this property. */
-	extern STRUCTUTILSEDITOR_API void RemoveCategory(FPropertyBagPropertyDesc& OutDesc, bool bAutoDisableCategories = true);
+	STRUCTUTILSEDITOR_API void RemoveCategory(FPropertyBagPropertyDesc& OutDesc, bool bAutoDisableCategories = true);
 	/** Get the category for this property. */
-	extern STRUCTUTILSEDITOR_API FString GetCategory(const FPropertyBagPropertyDesc& InDesc);
+	STRUCTUTILSEDITOR_API FString GetCategory(const FPropertyBagPropertyDesc& InDesc);
+
+	namespace Specifiers
+	{
+		extern STRUCTUTILSEDITOR_API const FLazyName ToolTipName;
+		extern STRUCTUTILSEDITOR_API const FLazyName ClampMinName;
+		extern STRUCTUTILSEDITOR_API const FLazyName ClampMaxName;
+		extern STRUCTUTILSEDITOR_API const FLazyName UIMinName;
+		extern STRUCTUTILSEDITOR_API const FLazyName UIMaxName;
+		extern STRUCTUTILSEDITOR_API const FLazyName UnitsName;
+
+		STRUCTUTILSEDITOR_API bool SupportsClamping(const FPropertyBagPropertyDesc& InDesc);
+		STRUCTUTILSEDITOR_API bool SupportsSettingUnits(const FPropertyBagPropertyDesc& InDesc);
+	}
 }
diff --git a/Engine/Source/Runtime/CoreUObject/Private/StructUtils/PropertyBag.cpp b/Engine/Source/Runtime/CoreUObject/Private/StructUtils/PropertyBag.cpp
--- a/Engine/Source/Runtime/CoreUObject/Private/StructUtils/PropertyBag.cpp
+++ b/Engine/Source/Runtime/CoreUObject/Private/StructUtils/PropertyBag.cpp
@@ -1877,6 +1877,42 @@ bool FPropertyBagPropertyDesc::CompatibleType(const FPropertyBagPropertyDesc& Ot
 	return true;
 }
 
+#if WITH_EDITOR
+bool FPropertyBagPropertyDesc::HasMetaData(FName SpecifierName) const
+{
+	return MetaData.ContainsByPredicate([SpecifierName](const FPropertyBagPropertyDescMetaData& MetaData)
+	{
+		return MetaData.Key == SpecifierName;
+	});
+}
+
+void FPropertyBagPropertyDesc::SetMetaData(FName SpecifierName, const FString& InValue)
+{
+	if (FPropertyBagPropertyDescMetaData* Specifier =
+		MetaData.FindByPredicate([SpecifierName](const FPropertyBagPropertyDescMetaData& InMetaData) { return InMetaData.Key == SpecifierName; }))
+	{
+		Specifier->Value = InValue;
+	}
+	else
+	{
+		MetaData.Emplace(FPropertyBagPropertyDescMetaData(SpecifierName, InValue));
+	}
+}
+
+FString FPropertyBagPropertyDesc::GetMetaData(FName SpecifierName) const
+{
+	if (const FPropertyBagPropertyDescMetaData* Specifier = MetaData.FindByPredicate(
+		[SpecifierName](const FPropertyBagPropertyDescMetaData& MetaData) { return MetaData.Key == SpecifierName; }))
+	{
+		return *Specifier->Value;
+	}
+	else
+	{
+		return {};
+	}
+}
+#endif // WITH_EDITOR
+
 //----------------------------------------------------------------//
 //  FInstancedPropertyBag
 //----------------------------------------------------------------//
diff --git a/Engine/Source/Runtime/CoreUObject/Public/StructUtils/PropertyBag.h b/Engine/Source/Runtime/CoreUObject/Public/StructUtils/PropertyBag.h
--- a/Engine/Source/Runtime/CoreUObject/Public/StructUtils/PropertyBag.h
+++ b/Engine/Source/Runtime/CoreUObject/Public/StructUtils/PropertyBag.h
@@ -307,7 +307,18 @@ struct FPropertyBagPropertyDesc
 	 * @param Ar The archive to serialize from or into.
 	 * @param Desc The PropertyBagPropertyDesc to serialize.
 	 */
-	 UE_API friend FArchive& operator<<(FArchive& Ar, FPropertyBagPropertyDesc& Desc);
+	UE_API friend FArchive& operator<<(FArchive& Ar, FPropertyBagPropertyDesc& Desc);
+
+#if WITH_EDITOR
+	/** @return true if the property has a metadata specifier */
+	UE_API bool HasMetaData(FName SpecifierName) const;
+
+	/** Set the metadata specifier value. Will add the specifier if it does not exist */
+	UE_API void SetMetaData(FName SpecifierName, const FString& InValue);
+
+	/** @return retrieves the metadata specifier value */
+	UE_API FString GetMetaData(FName SpecifierName) const;
+#endif // WITH_EDITOR
 
 	/** Pointer to object that defines the Enum, Struct, or Class. */
 	UPROPERTY(EditAnywhere, Category="Default")
