[InstancedStructs] Exposed FInstancedStructProvider

#rnx

diff --git a/Engine/Source/Editor/StructUtilsEditor/Private/InstancedStructDetails.cpp b/Engine/Source/Editor/StructUtilsEditor/Private/InstancedStructDetails.cpp
--- a/Engine/Source/Editor/StructUtilsEditor/Private/InstancedStructDetails.cpp
+++ b/Engine/Source/Editor/StructUtilsEditor/Private/InstancedStructDetails.cpp
@@ -24,139 +24,122 @@
 
 ////////////////////////////////////
 
-class FInstancedStructProvider : public IStructureDataProvider
+void FInstancedStructProvider::Reset()
 {
-public:
-	FInstancedStructProvider() = default;
-	
-	explicit FInstancedStructProvider(const TSharedPtr<IPropertyHandle>& InStructProperty)
-		: StructProperty(InStructProperty)
-	{
-	}
-	
-	virtual ~FInstancedStructProvider() override {}
+	StructProperty = nullptr;
+}
 
-	void Reset()
-	{
-		StructProperty = nullptr;
-	}
-	
-	virtual bool IsValid() const override
+bool FInstancedStructProvider::IsValid() const
+{
+	bool bHasValidData = false;
+	EnumerateInstances([&bHasValidData](const UScriptStruct* ScriptStruct, uint8* Memory, UPackage* Package)
 	{
-		bool bHasValidData = false;
-		EnumerateInstances([&bHasValidData](const UScriptStruct* ScriptStruct, uint8* Memory, UPackage* Package)
+		if (ScriptStruct && Memory)
 		{
-			if (ScriptStruct && Memory)
-			{
-				bHasValidData = true;
-				return false; // Stop
-			}
-			return true; // Continue
-		});
+			bHasValidData = true;
+			return false; // Stop
+		}
+		return true; // Continue
+	});
 
-		return bHasValidData;
-	}
-	
-	virtual const UStruct* GetBaseStructure() const override
+	return bHasValidData;
+}
+
+const UStruct* FInstancedStructProvider::GetBaseStructure() const
+{
+	// Taken from UClass::FindCommonBase
+	auto FindCommonBaseStruct = [](const UScriptStruct* StructA, const UScriptStruct* StructB)
 	{
-		// Taken from UClass::FindCommonBase
-		auto FindCommonBaseStruct = [](const UScriptStruct* StructA, const UScriptStruct* StructB)
+		const UScriptStruct* CommonBaseStruct = StructA;
+		while (CommonBaseStruct && StructB && !StructB->IsChildOf(CommonBaseStruct))
 		{
-			const UScriptStruct* CommonBaseStruct = StructA;
-			while (CommonBaseStruct && StructB && !StructB->IsChildOf(CommonBaseStruct))
-			{
-				CommonBaseStruct = Cast<UScriptStruct>(CommonBaseStruct->GetSuperStruct());
-			}
-			return CommonBaseStruct;
-		};
+			CommonBaseStruct = Cast<UScriptStruct>(CommonBaseStruct->GetSuperStruct());
+		}
+		return CommonBaseStruct;
+	};
 
-		const UScriptStruct* CommonStruct = nullptr;
-		EnumerateInstances([&CommonStruct, &FindCommonBaseStruct](const UScriptStruct* ScriptStruct, uint8* Memory, UPackage* Package)
+	const UScriptStruct* CommonStruct = nullptr;
+	EnumerateInstances([&CommonStruct, &FindCommonBaseStruct](const UScriptStruct* ScriptStruct, uint8* Memory, UPackage* Package)
+	{
+		if (ScriptStruct)
 		{
-			if (ScriptStruct)
-			{
-				CommonStruct = FindCommonBaseStruct(ScriptStruct, CommonStruct);
-			}
-			return true; // Continue
-		});
+			CommonStruct = FindCommonBaseStruct(ScriptStruct, CommonStruct);
+		}
+		return true; // Continue
+	});
 
-		return CommonStruct;
-	}
+	return CommonStruct;
+}
 
-	virtual void GetInstances(TArray<TSharedPtr<FStructOnScope>>& OutInstances, const UStruct* ExpectedBaseStructure) const override
+void FInstancedStructProvider::GetInstances(TArray<TSharedPtr<FStructOnScope>>& OutInstances, const UStruct* ExpectedBaseStructure) const
+{
+	// The returned instances need to be compatible with base structure.
+	// This function returns empty instances in case they are not compatible, with the idea that we have as many instances as we have outer objects.
+	EnumerateInstances([&OutInstances, ExpectedBaseStructure](const UScriptStruct* ScriptStruct, uint8* Memory, UPackage* Package)
 	{
-		// The returned instances need to be compatible with base structure.
-		// This function returns empty instances in case they are not compatible, with the idea that we have as many instances as we have outer objects.
-		EnumerateInstances([&OutInstances, ExpectedBaseStructure](const UScriptStruct* ScriptStruct, uint8* Memory, UPackage* Package)
+		TSharedPtr<FStructOnScope> Result;
+		
+		if (ExpectedBaseStructure && ScriptStruct && ScriptStruct->IsChildOf(ExpectedBaseStructure))
 		{
-			TSharedPtr<FStructOnScope> Result;
-			
-			if (ExpectedBaseStructure && ScriptStruct && ScriptStruct->IsChildOf(ExpectedBaseStructure))
-			{
-				Result = MakeShared<FStructOnScope>(ScriptStruct, Memory);
-				Result->SetPackage(Package);
-			}
+			Result = MakeShared<FStructOnScope>(ScriptStruct, Memory);
+			Result->SetPackage(Package);
+		}
 
-			OutInstances.Add(Result);
+		OutInstances.Add(Result);
 
-			return true; // Continue
-		});
-	}
+		return true; // Continue
+	});
+}
 
-	virtual bool IsPropertyIndirection() const override
+bool FInstancedStructProvider::IsPropertyIndirection() const
+{
+	return true;
+}
+
+uint8* FInstancedStructProvider::GetValueBaseAddress(uint8* ParentValueAddress, const UStruct* ExpectedBaseStructure) const
+{
+	if (!ParentValueAddress)
 	{
-		return true;
+		return nullptr;
 	}
 
-	virtual uint8* GetValueBaseAddress(uint8* ParentValueAddress, const UStruct* ExpectedBaseStructure) const override
+	FInstancedStruct& InstancedStruct = *reinterpret_cast<FInstancedStruct*>(ParentValueAddress);
+	if (ExpectedBaseStructure && InstancedStruct.GetScriptStruct() && InstancedStruct.GetScriptStruct()->IsChildOf(ExpectedBaseStructure))
 	{
-		if (!ParentValueAddress)
-		{
-			return nullptr;
-		}
-
-		FInstancedStruct& InstancedStruct = *reinterpret_cast<FInstancedStruct*>(ParentValueAddress);
-		if (ExpectedBaseStructure && InstancedStruct.GetScriptStruct() && InstancedStruct.GetScriptStruct()->IsChildOf(ExpectedBaseStructure))
-		{
-			return InstancedStruct.GetMutableMemory();
-		}
-		
-		return nullptr;
+		return InstancedStruct.GetMutableMemory();
 	}
 	
-protected:
+	return nullptr;
+}
 
-	void EnumerateInstances(TFunctionRef<bool(const UScriptStruct* ScriptStruct, uint8* Memory, UPackage* Package)> InFunc) const
+void FInstancedStructProvider::EnumerateInstances(TFunctionRef<bool(const UScriptStruct* ScriptStruct, uint8* Memory, UPackage* Package)> InFunc) const
+{
+	if (!StructProperty.IsValid())
 	{
-		if (!StructProperty.IsValid())
-		{
-			return;
-		}
-		
-		TArray<UPackage*> Packages;
-		StructProperty->GetOuterPackages(Packages);
+		return;
+	}
+	
+	TArray<UPackage*> Packages;
+	StructProperty->GetOuterPackages(Packages);
 
-		StructProperty->EnumerateRawData([&InFunc, &Packages](void* RawData, const int32 DataIndex, const int32 /*NumDatas*/)
+	StructProperty->EnumerateRawData([&InFunc, &Packages](void* RawData, const int32 DataIndex, const int32 /*NumDatas*/)
+	{
+		const UScriptStruct* ScriptStruct = nullptr;
+		uint8* Memory = nullptr;
+		UPackage* Package = nullptr;
+		if (FInstancedStruct* InstancedStruct = static_cast<FInstancedStruct*>(RawData))
 		{
-			const UScriptStruct* ScriptStruct = nullptr;
-			uint8* Memory = nullptr;
-			UPackage* Package = nullptr;
-			if (FInstancedStruct* InstancedStruct = static_cast<FInstancedStruct*>(RawData))
+			ScriptStruct = InstancedStruct->GetScriptStruct();
+			Memory = InstancedStruct->GetMutableMemory();
+			if (ensureMsgf(Packages.IsValidIndex(DataIndex), TEXT("Expecting packges and raw data to match.")))
 			{
-				ScriptStruct = InstancedStruct->GetScriptStruct();
-				Memory = InstancedStruct->GetMutableMemory();
-				if (ensureMsgf(Packages.IsValidIndex(DataIndex), TEXT("Expecting packges and raw data to match.")))
-				{
-					Package = Packages[DataIndex];
-				}
+				Package = Packages[DataIndex];
 			}
+		}
 
-			return InFunc(ScriptStruct, Memory, Package);
-		});
-	}
-	
-	TSharedPtr<IPropertyHandle> StructProperty;
-};
+		return InFunc(ScriptStruct, Memory, Package);
+	});
+}
 
 ////////////////////////////////////
 
diff --git a/Engine/Source/Editor/StructUtilsEditor/Public/InstancedStructDetails.h b/Engine/Source/Editor/StructUtilsEditor/Public/InstancedStructDetails.h
--- a/Engine/Source/Editor/StructUtilsEditor/Public/InstancedStructDetails.h
+++ b/Engine/Source/Editor/StructUtilsEditor/Public/InstancedStructDetails.h
@@ -4,6 +4,7 @@
 
 #include "IPropertyTypeCustomization.h"
 #include "IDetailCustomNodeBuilder.h"
+#include "IStructureDataProvider.h"
 
 #define UE_API STRUCTUTILSEDITOR_API
 
@@ -117,4 +118,31 @@ protected:
 	UE_API void OnStructLayoutChanges();
 };
 
+class FInstancedStructProvider : public IStructureDataProvider
+{
+public:
+	FInstancedStructProvider() = default;
+
+	explicit FInstancedStructProvider(const TSharedPtr<IPropertyHandle>& InStructProperty)
+		: StructProperty(InStructProperty)
+	{
+	}
+
+	virtual ~FInstancedStructProvider() override 
+	{
+	}
+
+	UE_API void Reset();
+	UE_API virtual bool IsValid() const override;
+	UE_API virtual const UStruct* GetBaseStructure() const override;
+	UE_API virtual void GetInstances(TArray<TSharedPtr<FStructOnScope>>& OutInstances, const UStruct* ExpectedBaseStructure) const override;
+	UE_API virtual bool IsPropertyIndirection() const override;
+	UE_API virtual uint8* GetValueBaseAddress(uint8* ParentValueAddress, const UStruct* ExpectedBaseStructure) const override;
+
+protected:
+	UE_API void EnumerateInstances(TFunctionRef<bool(const UScriptStruct* ScriptStruct, uint8* Memory, UPackage* Package)> InFunc) const;
+
+	TSharedPtr<IPropertyHandle> StructProperty;
+};
+
 #undef UE_API
