PropertyBagEditor: Add option for specifying ChildRowFeatures (as bitflag) and ShowOnlyInnerProperties via property instance metadata.

diff --git a/Engine/Source/Editor/StructUtilsEditor/Private/Customizations/PropertyBagDetails.cpp b/Engine/Source/Editor/StructUtilsEditor/Private/Customizations/PropertyBagDetails.cpp
--- a/Engine/Source/Editor/StructUtilsEditor/Private/Customizations/PropertyBagDetails.cpp
+++ b/Engine/Source/Editor/StructUtilsEditor/Private/Customizations/PropertyBagDetails.cpp
@@ -2326,11 +2326,17 @@ void FPropertyBagDetails::CustomizeHeader(TSharedRef<IPropertyHandle> StructProp
 		}
 
 		// Load the feature set by the metadata set on the FPropertyBag. Can only accept explicit enum values currently.
-		if (MetaDataProperty->HasMetaData(UE::StructUtils::Metadata::ChildRowFeaturesName))
+		const FString* ChildRowFeaturesString = StructProperty->GetInstanceMetaData(UE::StructUtils::Metadata::ChildRowFeaturesName);
+		if (!ChildRowFeaturesString)
+		{
+			ChildRowFeaturesString = MetaDataProperty->FindMetaData(UE::StructUtils::Metadata::ChildRowFeaturesName);
+		}
+
+		if (ChildRowFeaturesString)
 		{
 			if (const UEnum* Enum = StaticEnum<EPropertyBagChildRowFeatures>())
 			{
-				const int32 EnumIndex = Enum->GetIndexByNameString(MetaDataProperty->GetMetaData(UE::StructUtils::Metadata::ChildRowFeaturesName));
+				const int32 EnumIndex = Enum->GetIndexByNameString(*ChildRowFeaturesString);
 				if (EnumIndex != INDEX_NONE)
 				{
 					ChildRowFeatures = static_cast<EPropertyBagChildRowFeatures>(Enum->GetValueByIndex(EnumIndex));
@@ -2339,7 +2345,8 @@ void FPropertyBagDetails::CustomizeHeader(TSharedRef<IPropertyHandle> StructProp
 		}
 
 		// Don't show the header if ShowOnlyInnerProperties is set
-		if (MetaDataProperty->HasMetaData(UE::StructUtils::Metadata::ShowOnlyInnerPropertiesName))
+		if (MetaDataProperty->HasMetaData(UE::StructUtils::Metadata::ShowOnlyInnerPropertiesName) ||
+			StructProperty->GetInstanceMetaData(UE::StructUtils::Metadata::ShowOnlyInnerPropertiesName))
 		{
 			return;
 		}
