PropertyBag: Fix Category Delete button was visible for Fixed bags too.

diff --git a/Engine/Source/Editor/StructUtilsEditor/Private/Customizations/PropertyBagDetails.cpp b/Engine/Source/Editor/StructUtilsEditor/Private/Customizations/PropertyBagDetails.cpp
--- a/Engine/Source/Editor/StructUtilsEditor/Private/Customizations/PropertyBagDetails.cpp
+++ b/Engine/Source/Editor/StructUtilsEditor/Private/Customizations/PropertyBagDetails.cpp
@@ -1207,6 +1207,7 @@ void FPropertyBagInstanceDataDetails::OnGroupRowAdded(IDetailGroup& GroupRow, in
 				.VAlign(VAlign_Center)
 				.ButtonStyle(FAppStyle::Get(), "SimpleButton")
 				.ToolTipText(LOCTEXT("DeleteCategory", "Delete this category."))
+				.Visibility(bIsFixed ? EVisibility::Collapsed : EVisibility::Visible)
 				.OnClicked_Lambda([GroupName = GroupRow.GetGroupName(), BagStructProperty = BagStructProperty, PropUtils = PropUtils]()
 				{
 					Private::ApplyChangesToPropertyDescs(
