Expose AddChildRows via InstancedStructDataDetails so that it's easier to customize the way in which children are laid out.

diff --git a/Engine/Source/Editor/StructUtilsEditor/Private/InstancedStructDetails.cpp b/Engine/Source/Editor/StructUtilsEditor/Private/InstancedStructDetails.cpp
--- a/Engine/Source/Editor/StructUtilsEditor/Private/InstancedStructDetails.cpp
+++ b/Engine/Source/Editor/StructUtilsEditor/Private/InstancedStructDetails.cpp
@@ -159,6 +159,40 @@ FInstancedStructDataDetails::~FInstancedStructDataDetails()
 	UE::StructUtils::Delegates::OnUserDefinedStructReinstanced.Remove(UserDefinedStructReinstancedHandle);
 }
 
+TArray<FString> FInstancedStructDataDetails::GetPropertyCategories(TSharedPtr<IPropertyHandle> PropertyHandle)
+{
+	static const FName CategoryName = FName(TEXT("Category"));
+	static const FName EnableCategoriesName = FName(TEXT("EnableCategories"));
+
+	// The property needs the "EnableCategories" metadata in order to be added under a group. Grouping is opt-in.
+	if (!PropertyHandle->HasMetaData(EnableCategoriesName))
+	{
+		return { };
+	}
+
+	const FString& PropertyCategory = PropertyHandle->GetMetaData(CategoryName).TrimStartAndEnd();
+	if (PropertyCategory.IsEmpty())
+	{
+		return { };
+	}
+
+	constexpr bool bCullEmpty = true;
+	TArray<FString> CategoriesToAdd;
+	PropertyCategory.ParseIntoArray(CategoriesToAdd, TEXT("|"), bCullEmpty);
+
+	// Clean up categories
+	for (int32 Index = 0; Index < CategoriesToAdd.Num(); ++Index)
+	{
+		FString& CategoryToAdd = CategoriesToAdd[Index];
+		CategoryToAdd.TrimStartAndEndInline();
+
+		// Cover the edge case where there's a category like "Foo|", which is invalid
+		CategoryToAdd.TrimCharInline(TEXT('|'), nullptr);
+	}
+
+	return CategoriesToAdd;
+}
+
 void FInstancedStructDataDetails::OnUserDefinedStructReinstancedHandle(const UUserDefinedStruct& Struct)
 {
 	OnStructLayoutChanges();
@@ -195,30 +229,12 @@ void FInstancedStructDataDetails::GetPropertyGroups(
 	IDetailChildrenBuilder& InChildBuilder,
 	TMap<TSharedPtr<IPropertyHandle>, IDetailGroup*>& OutPropertyToGroup) const
 {
-	static const FName CategoryName = FName(TEXT("Category"));
-	static const FName EnableCategoriesName = FName(TEXT("EnableCategories"));
-
 	// Temporarily store a mapping of category -> group while groups are being built
 	TMap<FString, IDetailGroup*> CategoryToGroup;
 	
 	for (const TSharedPtr<IPropertyHandle>& PropertyHandle : InProperties)
 	{
-		// The property needs the "EnableCategories" metadata in order to be added under a group. Grouping is opt-in.
-		if (!PropertyHandle->HasMetaData(EnableCategoriesName))
-		{
-			continue;
-		}
-		
-		const FString& PropertyCategory = PropertyHandle->GetMetaData(CategoryName).TrimStartAndEnd();
-		if (PropertyCategory.IsEmpty())
-		{
-			continue;
-		}
-
-		constexpr bool bCullEmpty = true;
-		TArray<FString> CategoriesToAdd;
-		PropertyCategory.ParseIntoArray(CategoriesToAdd, TEXT("|"), bCullEmpty);
-
+		TArray<FString> CategoriesToAdd = GetPropertyCategories(PropertyHandle);
 		if (CategoriesToAdd.IsEmpty())
 		{
 			continue;
@@ -232,11 +248,6 @@ void FInstancedStructDataDetails::GetPropertyGroups(
 		for (int32 Index = 0; Index < CategoriesToAdd.Num(); ++Index)
 		{
 			FString& CategoryToAdd = CategoriesToAdd[Index];
-			CategoryToAdd.TrimStartAndEndInline();
-
-			// Cover the edge case where there's a category like "Foo|", which is invalid
-			CategoryToAdd.TrimCharInline(TEXT('|'), nullptr);
-
 			CompleteCategory = CompleteCategory.IsEmpty()
 				? CategoryToAdd
 				: FString::Join(TArray({CompleteCategory, CategoryToAdd}), TEXT("|"));
@@ -332,26 +343,7 @@ void FInstancedStructDataDetails::GenerateChildContent(IDetailChildrenBuilder& C
 	{
 		StructProperty->RemoveChildren();
 		TArray<TSharedPtr<IPropertyHandle>> ChildProperties = StructProperty->AddChildStructure(NewStructProvider);
-
-		// Properties may have Category metadata. If that's the case, they should be added under groups.
-		TMap<TSharedPtr<IPropertyHandle>, IDetailGroup*> PropertyToGroup;
-		GetPropertyGroups(ChildProperties, ChildBuilder, PropertyToGroup);
-		
-		for (TSharedPtr<IPropertyHandle> ChildHandle : ChildProperties)
-		{
-			// If the property has a group, add it under the group. Otherwise, just add it normally via the builder.
-			IDetailGroup** PropertyGroup = PropertyToGroup.Find(ChildHandle);
-			if (PropertyGroup && *PropertyGroup)
-			{
-				IDetailPropertyRow& Row = (*PropertyGroup)->AddPropertyRow(ChildHandle.ToSharedRef());
-				OnChildRowAdded(Row);
-			}
-			else
-			{
-				IDetailPropertyRow& Row = ChildBuilder.AddProperty(ChildHandle.ToSharedRef());
-				OnChildRowAdded(Row);
-			}
-		}
+		AddChildRows(ChildBuilder, ChildProperties);
 	}
 
 	bCanHandleStructValuePostChange = true;
@@ -359,6 +351,30 @@ void FInstancedStructDataDetails::GenerateChildContent(IDetailChildrenBuilder& C
 	CachedInstanceTypes = GetInstanceTypes();
 }
 
+void FInstancedStructDataDetails::AddChildRows(IDetailChildrenBuilder& ChildBuilder,
+	const TArray<TSharedPtr<IPropertyHandle>>& ChildProperties)
+{
+	// Properties may have Category metadata. If that's the case, they should be added under groups.
+	TMap<TSharedPtr<IPropertyHandle>, IDetailGroup*> PropertyToGroup;
+	GetPropertyGroups(ChildProperties, ChildBuilder, PropertyToGroup);
+
+	for (TSharedPtr<IPropertyHandle> ChildHandle : ChildProperties)
+	{
+		// If the property has a group, add it under the group. Otherwise, just add it normally via the builder.
+		IDetailGroup** PropertyGroup = PropertyToGroup.Find(ChildHandle);
+		if (PropertyGroup && *PropertyGroup)
+		{
+			IDetailPropertyRow& Row = (*PropertyGroup)->AddPropertyRow(ChildHandle.ToSharedRef());
+			OnChildRowAdded(Row);
+		}
+		else
+		{
+			IDetailPropertyRow& Row = ChildBuilder.AddProperty(ChildHandle.ToSharedRef());
+			OnChildRowAdded(Row);
+		}
+	}
+}
+
 void FInstancedStructDataDetails::Tick(float DeltaTime)
 {
 	// If the instance types change (e.g. due to selecting new struct type), we'll need to update the layout.
diff --git a/Engine/Source/Editor/StructUtilsEditor/Public/InstancedStructDetails.h b/Engine/Source/Editor/StructUtilsEditor/Public/InstancedStructDetails.h
--- a/Engine/Source/Editor/StructUtilsEditor/Public/InstancedStructDetails.h
+++ b/Engine/Source/Editor/StructUtilsEditor/Public/InstancedStructDetails.h
@@ -81,6 +81,10 @@ public:
 	// Called when a child is added, override to customize a child row.
 	virtual void OnChildRowAdded(IDetailPropertyRow& ChildRow) {}
 
+protected:
+	UE_API static TArray<FString> GetPropertyCategories(TSharedPtr<IPropertyHandle> PropertyHandle);
+	UE_API virtual void AddChildRows(IDetailChildrenBuilder& ChildBuilder, const TArray<TSharedPtr<IPropertyHandle>>& ChildProperties);
+
 private:
 	UE_API void OnUserDefinedStructReinstancedHandle(const class UUserDefinedStruct& Struct);
 
