[StructUtils] PR #13435: Support determining base struct dynamically

#rnx

diff --git a/Engine/Source/Editor/StructUtilsEditor/Private/Customizations/PropertyBagDetails.cpp b/Engine/Source/Editor/StructUtilsEditor/Private/Customizations/PropertyBagDetails.cpp
--- a/Engine/Source/Editor/StructUtilsEditor/Private/Customizations/PropertyBagDetails.cpp
+++ b/Engine/Source/Editor/StructUtilsEditor/Private/Customizations/PropertyBagDetails.cpp
@@ -2,6 +2,8 @@
 
 #include "PropertyBagDetails.h"
 
+#include "StructUtilsEditorUtils.h"
+
 #include "DetailLayoutBuilder.h"
 #include "DetailWidgetRow.h"
 #include "IDetailChildrenBuilder.h"
@@ -56,7 +58,7 @@ void SetPropertyDescFromPin(FPropertyBagPropertyDesc& Desc, const FEdGraphPinTyp
 	default:
 		break;
 	}
-	
+
 	// Value type
 	if (PinType.PinCategory == UEdGraphSchema_K2::PC_Boolean)
 	{
@@ -97,7 +99,7 @@ void SetPropertyDescFromPin(FPropertyBagPropertyDesc& Desc, const FEdGraphPinTyp
 		{
 			Desc.ValueType = EPropertyBagPropertyType::Double;
 			Desc.ValueTypeObject = nullptr;
-		}		
+		}
 	}
 	else if (PinType.PinCategory == UEdGraphSchema_K2::PC_Name)
 	{
@@ -249,20 +251,6 @@ FEdGraphPinType GetPropertyDescAsPin(const FPropertyBagPropertyDesc& Desc)
 
 namespace Private
 {
-
-/** @return true property handle holds struct property of type T.  */ 
-template<typename T>
-bool IsScriptStruct(const TSharedPtr<IPropertyHandle>& PropertyHandle)
-{
-	if (!PropertyHandle)
-	{
-		return false;
-	}
-
-	FStructProperty* StructProperty = CastField<FStructProperty>(PropertyHandle->GetProperty());
-	return StructProperty && StructProperty->Struct->IsA(TBaseStructure<T>::Get()->GetClass());
-}
-
 /** @return true if the property is one of the known missing types. */
 bool HasMissingType(const TSharedPtr<IPropertyHandle>& PropertyHandle)
 {
@@ -324,12 +312,12 @@ const UPropertyBag* GetCommonBagStruct(TSharedPtr<IPropertyHandle> StructPropert
 TArray<FPropertyBagPropertyDesc> GetCommonPropertyDescs(const TSharedPtr<IPropertyHandle>& StructProperty)
 {
 	TArray<FPropertyBagPropertyDesc> PropertyDescs;
-	
+
 	if (const UPropertyBag* BagStruct = GetCommonBagStruct(StructProperty))
 	{
 		PropertyDescs = BagStruct->GetPropertyDescs();
 	}
-	
+
 	return PropertyDescs;
 }
 
@@ -374,12 +362,12 @@ bool IsUniqueName(const FName NewName, const FName OldName, const TSharedPtr<IPr
 	{
 		return false;
 	}
-	
+
 	if (!StructProperty || !StructProperty->IsValidHandle())
 	{
 		return false;
 	}
-	
+
 	bool bFound = false;
 
 	if (ensure(IsScriptStruct<FInstancedPropertyBag>(StructProperty)))
@@ -405,7 +393,7 @@ bool IsUniqueName(const FName NewName, const FName OldName, const TSharedPtr<IPr
 			return true;
 		});
 	}
-	
+
 	return !bFound;
 }
 
@@ -416,7 +404,7 @@ void ApplyChangesToPropertyDescs(const FText& SessionName, const TSharedPtr<IPro
 	{
 		return;
 	}
-	
+
 	FScopedTransaction Transaction(SessionName);
 	TArray<FPropertyBagPropertyDesc> PropertyDescs = GetCommonPropertyDescs(StructProperty);
 	StructProperty->NotifyPreChange();
@@ -424,7 +412,7 @@ void ApplyChangesToPropertyDescs(const FText& SessionName, const TSharedPtr<IPro
 	Function(PropertyDescs);
 
 	SetPropertyDescs(StructProperty, PropertyDescs);
-	
+
 	StructProperty->NotifyPostChange(EPropertyChangeType::ValueSet);
 	StructProperty->NotifyFinishedChangingProperties();
 	if (PropUtils)
@@ -459,7 +447,7 @@ void ApplyChangesToSinglePropertyDesc(const FText& SessionName, const FPropertyB
 
 bool CanHaveMemberVariableOfType(const FEdGraphPinType& PinType)
 {
-	if (PinType.PinCategory == UEdGraphSchema_K2::PC_Exec 
+	if (PinType.PinCategory == UEdGraphSchema_K2::PC_Exec
 		|| PinType.PinCategory == UEdGraphSchema_K2::PC_Wildcard
 		|| PinType.PinCategory == UEdGraphSchema_K2::PC_MCDelegate
 		|| PinType.PinCategory == UEdGraphSchema_K2::PC_Delegate
@@ -467,56 +455,8 @@ bool CanHaveMemberVariableOfType(const FEdGraphPinType& PinType)
 	{
 		return false;
 	}
-	
-	return true;
-}
-
-bool FindUserFunction(TSharedPtr<IPropertyHandle> InStructProperty, FName InFuncMetadataName, UFunction*& OutFunc, UObject*& OutTarget)
-{
-	FProperty* MetadataProperty = InStructProperty->GetMetaDataProperty();
-
-	OutFunc = nullptr;
-	OutTarget = nullptr;
-
-	if (!MetadataProperty || !MetadataProperty->HasMetaData(InFuncMetadataName))
-	{
-		return false;
-	}
-
-	FString FunctionName = MetadataProperty->GetMetaData(InFuncMetadataName);
-	if (FunctionName.IsEmpty())
-	{
-		return false;
-	}
-
-	TArray<UObject*> OutObjects;
-	InStructProperty->GetOuterObjects(OutObjects);
-
-	// Check for external function references, taken from GetOptions
-	if (FunctionName.Contains(TEXT(".")))
-	{
-		OutFunc = FindObject<UFunction>(nullptr, *FunctionName, EFindObjectFlags::ExactClass);
-
-		if (ensureMsgf(OutFunc && OutFunc->HasAnyFunctionFlags(EFunctionFlags::FUNC_Static), TEXT("[%s] Didn't find function %s or expected it to be static"), *InFuncMetadataName.ToString(), *FunctionName))
-		{
-			UObject* GetOptionsCDO = OutFunc->GetOuterUClass()->GetDefaultObject();
-			OutTarget = GetOptionsCDO;
-		}
-	}
-	else if (OutObjects.Num() > 0)
-	{
-		OutTarget = OutObjects[0];
-		OutFunc = OutTarget->GetClass() ? OutTarget->GetClass()->FindFunctionByName(*FunctionName) : nullptr;
-	}
-
-	// Only support native functions
-	if (!ensureMsgf(OutFunc && OutFunc->IsNative(), TEXT("[%s] Didn't find function %s or expected it to be native"), *InFuncMetadataName.ToString(), *FunctionName))
-	{
-		OutFunc = nullptr;
-		OutTarget = nullptr;
-	}
 
-	return OutTarget != nullptr && OutFunc != nullptr;
+	return true;
 }
 
 FText GetAccessSpecifierNameFromFlags(const EPropertyFlags Flags)
@@ -532,219 +472,6 @@ FText GetAccessSpecifierNameFromFlags(const EPropertyFlags Flags)
 	}
 }
 
-// UFunction calling helpers. 
-// Use our "own" hardcoded reflection system for types used in UFunctions calls in this file.
-template<typename T> struct TypeName { static const TCHAR* Get() = delete; };
-
-#define DEFINE_TYPENAME(InType) template<> struct TypeName<InType> { static const TCHAR *Get() { return TEXT(#InType); }};
-DEFINE_TYPENAME(bool)
-DEFINE_TYPENAME(FGuid)
-DEFINE_TYPENAME(FName)
-DEFINE_TYPENAME(FEdGraphPinType)
-#undef DEFINE_TYPENAME
-
-// Wrapper around a param that store an address (const_cast for const ptr, be careful of that)
-// a string identifiying the underlying cpp type and if the input value is const, mark it const.
-struct FFuncParam
-{
-	void* Value = nullptr;
-	const TCHAR* CppType;
-	bool bIsConst = false;
-
-	template<typename T>
-	static FFuncParam Make(T& Value)
-	{
-		FFuncParam Result;
-		Result.Value = &Value;
-		Result.CppType = TypeName<T>::Get();
-		return Result;
-	}
-
-	template<typename T>
-	static FFuncParam Make(const T& Value)
-	{
-		FFuncParam Result;
-		Result.Value = &const_cast<T&>(Value);
-		Result.CppType = TypeName<T>::Get();
-		Result.bIsConst = true;
-		return Result;
-	}
-};
-
-// Validate that the function pass as parameter has signature ReturnType(ArgsTypes...)
-template <typename ReturnType, typename... ArgsTypes>
-bool ValidateFunctionSignature(UFunction* InFunc)
-{
-	if (!InFunc)
-	{
-		return false;
-	}
-
-	constexpr int32 NumParms = std::is_same_v <ReturnType, void> ? sizeof...(ArgsTypes) : (sizeof...(ArgsTypes) + 1);
-
-	if (NumParms != InFunc->NumParms)
-	{
-		return false;
-	}
-
-	const TCHAR* ArgsCppTypes[NumParms] = { TypeName<ArgsTypes>::Get() ... };
-	
-	// If we have a return type, put it at the end. UFunction will have the return type after InArgs in the field iterator.
-	if constexpr (!std::is_same_v<ReturnType, void>)
-	{
-		ArgsCppTypes[NumParms - 1] = TypeName<ReturnType>::Get();
-	}
-	else
-	{
-		// Otherwise, check that the function doesn't have a return param
-		if (InFunc->GetReturnProperty() != nullptr)
-		{
-			return false;
-		}
-	}
-
-	int32 ArgCppTypesIndex = 0;
-	for (TFieldIterator<FProperty> It(InFunc); It && It->HasAnyPropertyFlags(EPropertyFlags::CPF_Parm); ++It)
-	{
-		const FString PropertyCppType = It->GetCPPType();
-		if (PropertyCppType != ArgsCppTypes[ArgCppTypesIndex])
-		{
-			return false;
-		}
-
-		// Also making sure that the last param is a return param, if we have a return value
-		if constexpr (!std::is_same_v<ReturnType, void>)
-		{
-			if (ArgCppTypesIndex == NumParms - 1 && !It->HasAnyPropertyFlags(EPropertyFlags::CPF_ReturnParm))
-			{
-				return false;
-			}
-		}
-
-		ArgCppTypesIndex++;
-	}
-
-	return true;
-}
-
-template <typename ReturnType, typename... ArgsTypes>
-TValueOrError<ReturnType, void> CallFunc(UObject* InTargetObject, UFunction* InFunc, ArgsTypes&& ...InArgs)
-{
-	if (!InTargetObject || !InFunc)
-	{
-		return MakeError();
-	}
-
-	constexpr int32 NumParms = std::is_same_v <ReturnType, void> ? sizeof...(ArgsTypes) : (sizeof...(ArgsTypes) + 1);
-
-	if (NumParms != InFunc->NumParms)
-	{
-		return MakeError();
-	}
-
-	FFuncParam InParams[NumParms] = { FFuncParam::Make(std::forward<ArgsTypes>(InArgs)) ... };
-
-	auto Invoke = [InTargetObject, InFunc, &InParams, NumParms](ReturnType* OutResult) -> bool
-	{
-		// Validate that the function has a return property if the return type is not void.
-		if (std::is_same_v<ReturnType, void> != (InFunc->GetReturnProperty() == nullptr))
-		{
-			return false;
-		}
-
-		// Allocating our "stack" for the function call on the stack (will be freed when this function is exited)
-		uint8* StackMemory = (uint8*)FMemory_Alloca(InFunc->ParmsSize);
-		FMemory::Memzero(StackMemory, InFunc->ParmsSize);
-
-		if constexpr (!std::is_same_v<ReturnType, void>)
-		{
-			check(OutResult != nullptr);
-			InParams[NumParms - 1] = FFuncParam::Make(*OutResult);
-		}
-
-		bool bValid = true;
-		int32 ParamIndex = 0;
-
-		// Initializing our "stack" with our parameters. Use the property system to make sure more complex types
-		// are constructed before being set.
-		for (TFieldIterator<FProperty> It(InFunc); It && It->HasAnyPropertyFlags(CPF_Parm); ++It)
-		{
-			FProperty* LocalProp = *It;
-			if (!LocalProp->HasAnyPropertyFlags(CPF_ZeroConstructor))
-			{
-				LocalProp->InitializeValue_InContainer(StackMemory);
-			}
-
-			if (bValid)
-			{
-				if (ParamIndex >= NumParms)
-				{
-					bValid = false;
-					continue;
-				}
-
-				FFuncParam& Param = InParams[ParamIndex++];
-
-				if (LocalProp->GetCPPType() != Param.CppType)
-				{
-					bValid = false;
-					continue;
-				}
-
-				LocalProp->SetValue_InContainer(StackMemory, Param.Value);
-			}
-		}
-
-		if (bValid)
-		{
-			FFrame Stack(InTargetObject, InFunc, StackMemory, nullptr, InFunc->ChildProperties);
-			InFunc->Invoke(InTargetObject, Stack, OutResult);
-		}
-
-		ParamIndex = 0;
-		// Copy back all non-const out params (that is not the return param, this one is already set by the invoke call)
-		// from the stack, also making sure that the constructed types are destroyed accordingly.
-		for (TFieldIterator<FProperty> It(InFunc); It && It->HasAnyPropertyFlags(CPF_Parm); ++It)
-		{
-			FProperty* LocalProp = *It;
-			FFuncParam& Param = InParams[ParamIndex++];
-
-			if (LocalProp->HasAnyPropertyFlags(CPF_OutParm) && !LocalProp->HasAnyPropertyFlags(CPF_ReturnParm) && !Param.bIsConst)
-			{
-				LocalProp->GetValue_InContainer(StackMemory, Param.Value);
-			}
-
-			LocalProp->DestroyValue_InContainer(StackMemory);
-		}
-
-		return bValid;
-	};
-
-	if constexpr (std::is_same_v<ReturnType, void>)
-	{
-		if (Invoke(nullptr))
-		{
-			return MakeValue();
-		}
-		else
-		{
-			return MakeError();
-		}
-	}
-	else
-	{
-		ReturnType OutResult{};
-		if (Invoke(&OutResult))
-		{
-			return MakeValue(std::move(OutResult));
-		}
-		else
-		{
-			return MakeError();
-		}
-	}
-}
-
 /** Checks if the value for a source property in a source struct has the same value that the target property in the target struct. */
 bool ArePropertiesIdentical(
 	const FPropertyBagPropertyDesc* InSourcePropertyDesc,
@@ -801,26 +528,20 @@ void GetFilteredVariableTypeTree(const TSharedPtr<IPropertyHandle>& BagStructPro
 		return;
 	}
 
-	UFunction* IsPinTypeAcceptedFunc = nullptr;
-	UObject* IsPinTypeAcceptedTarget = nullptr;
-	if (FindUserFunction(BagStructProperty, Metadata::IsPinTypeAcceptedName, IsPinTypeAcceptedFunc, IsPinTypeAcceptedTarget))
-	{
-		check(IsPinTypeAcceptedFunc && IsPinTypeAcceptedTarget);
+	DECLARE_DELEGATE_RetVal_TwoParams(bool, FIsPinTypeAccepted, FEdGraphPinType, bool);
+	FIsPinTypeAccepted IsPinTypeAcceptedDelegate{};
 
-		// We need to make sure the signature matches perfectly: bool(FEdGraphPinType, bool)
-		bool bFuncIsValid = ValidateFunctionSignature<bool, FEdGraphPinType, bool>(IsPinTypeAcceptedFunc);
-		if (!ensureMsgf(bFuncIsValid, TEXT("[%s] Function %s does not have the right signature."), *Metadata::IsPinTypeAcceptedName.ToString(), *IsPinTypeAcceptedFunc->GetName()))
-		{
-			return;
-		}
+	if (TOptional<FFindUserFunctionResult> Result = FindUserFunction(BagStructProperty, Metadata::IsPinTypeAcceptedName); Result.IsSet())
+	{
+		check(Result.GetValue().Function && Result.GetValue().Target);
+		IsPinTypeAcceptedDelegate = FIsPinTypeAccepted::CreateUFunction(Result.GetValue().Target, Result.GetValue().Function->GetFName());
 	}
 
-	auto IsPinTypeAccepted = [IsPinTypeAcceptedFunc, IsPinTypeAcceptedTarget](const FEdGraphPinType& InPinType, bool bInIsChild) -> bool
+	auto IsPinTypeAccepted = [&IsPinTypeAcceptedDelegate](const FEdGraphPinType& InPinType, bool bInIsChild) -> bool
 	{
-		if (IsPinTypeAcceptedFunc && IsPinTypeAcceptedTarget)
+		if (IsPinTypeAcceptedDelegate.IsBound())
 		{
-			const TValueOrError<bool, void> bIsValid = CallFunc<bool>(IsPinTypeAcceptedTarget, IsPinTypeAcceptedFunc, InPinType, bInIsChild);
-			return bIsValid.HasValue() && bIsValid.GetValue();
+			return IsPinTypeAcceptedDelegate.Execute(InPinType, bInIsChild);
 		}
 		else
 		{
@@ -869,11 +590,9 @@ bool CanDeleteProperty(const TSharedPtr<IPropertyHandle>& InStructProperty, cons
 
 	// Extra check provided by the user to cancel a remove action. Useful to provide the user a possibility to cancel the action if
 	// the given property is in use elsewhere.
-	UFunction* CanRemovePropertyFunc = nullptr;
-	UObject* CanRemovePropertyTarget = nullptr;
-	if (FindUserFunction(InStructProperty, Metadata::CanRemovePropertyName, CanRemovePropertyFunc, CanRemovePropertyTarget))
+	if (TOptional<FFindUserFunctionResult> Result = FindUserFunction(InStructProperty, Metadata::CanRemovePropertyName); Result.IsSet())
 	{
-		check(CanRemovePropertyFunc && CanRemovePropertyTarget);
+		check(Result.GetValue().Function && Result.GetValue().Target);
 
 		FName PropertyName = ChildPropertyHandle->GetProperty()->GetFName();
 		const UPropertyBag* PropertyBag = GetCommonBagStruct(InStructProperty);
@@ -884,22 +603,13 @@ bool CanDeleteProperty(const TSharedPtr<IPropertyHandle>& InStructProperty, cons
 			return false;
 		}
 
-		// We need to make sure the signature matches perfectly: bool(FGuid, FName)s
-		bool bFuncIsValid = UE::StructUtils::Private::ValidateFunctionSignature<bool, FGuid, FName>(CanRemovePropertyFunc);
-		if (!ensureMsgf(bFuncIsValid, TEXT("[%s] Function %s does not have the right signature."), *UE::StructUtils::Metadata::CanRemovePropertyName.ToString(), *CanRemovePropertyFunc->GetName()))
-		{
-			return false;
-		}
-
-		const TValueOrError<bool, void> bCanRemove = UE::StructUtils::Private::CallFunc<bool>(CanRemovePropertyTarget, CanRemovePropertyFunc, PropertyDesc->ID, PropertyDesc->Name);
-
-		if (bCanRemove.HasError() || !bCanRemove.GetValue())
-		{
-			return false;
-		}
+		DECLARE_DELEGATE_RetVal_TwoParams(bool, FGetCanDeleteProperty, FGuid, FName);
+		return FGetCanDeleteProperty::CreateUFunction(Result.GetValue().Target, Result.GetValue().Function->GetFName()).Execute(PropertyDesc->ID, PropertyDesc->Name);
+	}
+	else
+	{
+		return true;
 	}
-
-	return true;
 }
 
 void DeleteProperty(TSharedPtr<IPropertyHandle> InStructProperty, TSharedPtr<IPropertyHandle> ChildPropertyHandle, const TSharedPtr<IPropertyUtilities>& PropUtils)
@@ -1010,7 +720,7 @@ namespace Constants
 //----------------------------------------------------------------//
 //  FPropertyBagInstanceDataDetails
 //  - StructProperty is FInstancedPropertyBag
-//  - ChildPropertyHandle a child property of the FInstancedPropertyBag::Value (FInstancedStruct)  
+//  - ChildPropertyHandle a child property of the FInstancedPropertyBag::Value (FInstancedStruct)
 //----------------------------------------------------------------//
 
 /** Primary constructor. Values passed by parameter struct. */
@@ -1342,7 +1052,7 @@ void FPropertyBagInstanceDataDetails::OnChildRowAdded(IDetailPropertyRow& ChildR
 				}
 				return true;
 			});
-		
+
 		FOnBooleanValueChanged OnEditConditionChanged = FOnBooleanValueChanged::CreateLambda([WeakSelf, ChildPropertyHandle](bool bNewValue)
 		{
 			if (const TSharedPtr<FPropertyBagInstanceDataDetails> Self = WeakSelf.Pin())
@@ -2056,7 +1766,7 @@ FPropertyBagInstanceDataDetails::EPropertyOverrideState FPropertyBagInstanceData
 	}
 
 	int32 NumValues = 0;
-	int32 NumOverrides = 0; 
+	int32 NumOverrides = 0;
 
 	const FProperty* Property = ChildPropertyHandle->GetProperty();
 	check(Property);
@@ -2077,7 +1787,7 @@ FPropertyBagInstanceDataDetails::EPropertyOverrideState FPropertyBagInstanceData
 
 			return true;
 		});
-	
+
 	if (NumOverrides == 0)
 	{
 		return EPropertyOverrideState::No;
@@ -2085,10 +1795,10 @@ FPropertyBagInstanceDataDetails::EPropertyOverrideState FPropertyBagInstanceData
 	else if (NumOverrides == NumValues)
 	{
 		return EPropertyOverrideState::Yes;
-	}			
+	}
 	return EPropertyOverrideState::Undetermined;
 }
-	
+
 void FPropertyBagInstanceDataDetails::SetPropertyOverride(TSharedPtr<IPropertyHandle> ChildPropertyHandle, const bool bIsOverridden)
 {
 	if (!ChildPropertyHandle)
@@ -2100,9 +1810,9 @@ void FPropertyBagInstanceDataDetails::SetPropertyOverride(TSharedPtr<IPropertyHa
 	check(Property);
 
 	FScopedTransaction Transaction(FText::Format(LOCTEXT("OverrideChange", "Change Override for {0}"), FText::FromName(ChildPropertyHandle->GetProperty()->GetFName())));
-	
+
 	PreChangeOverrides();
-	
+
 	EnumeratePropertyBags(
 		BagStructProperty,
 		[Property, bIsOverridden]
@@ -2167,7 +1877,7 @@ bool FPropertyBagInstanceDataDetails::IsDefaultValue(TSharedPtr<IPropertyHandle>
 	{
 		return true;
 	}
-	
+
 	return false;
 }
 
@@ -2177,13 +1887,13 @@ void FPropertyBagInstanceDataDetails::ResetToDefault(TSharedPtr<IPropertyHandle>
 	{
 		return;
 	}
-	
+
 	const FProperty* Property = ChildPropertyHandle->GetProperty();
 	check(Property);
 
 	FScopedTransaction Transaction(FText::Format(LOCTEXT("ResetToDefault", "Reset {0} to default value"), FText::FromName(ChildPropertyHandle->GetProperty()->GetFName())));
 	ChildPropertyHandle->NotifyPreChange();
-	
+
 	EnumeratePropertyBags(
 		BagStructProperty,
 		[Property]
@@ -2220,7 +1930,7 @@ TSharedRef<SWidget> FPropertyBagInstanceDataDetails::OnPropertyNameContent(TShar
 		{
 			return;
 		}
-		
+
 		UE::StructUtils::Private::ApplyChangesToPropertyDescs(
 		LOCTEXT("OnPropertyMoved", "Move Property"), BagStructProperty, PropUtils,
 		[&ChildPropertyHandle, &Delta](TArray<FPropertyBagPropertyDesc>& PropertyDescs)
@@ -2251,9 +1961,9 @@ TSharedRef<SWidget> FPropertyBagInstanceDataDetails::OnPropertyNameContent(TShar
 				bAllowContainers)
 		],
 		FText::GetEmpty());
-	
+
 	MenuBuilder.AddSeparator();
-	
+
 	MenuBuilder.AddMenuEntry(
 		LOCTEXT("Rename", "Rename"),
 		LOCTEXT("Rename_ToolTip", "Rename property"),
@@ -2270,7 +1980,7 @@ TSharedRef<SWidget> FPropertyBagInstanceDataDetails::OnPropertyNameContent(TShar
 			UE::StructUtils::Private::DeleteProperty(BagStructProperty, ChildPropertyHandle, PropUtils);
 		}))
 	);
-	
+
 	MenuBuilder.AddSeparator();
 
 	MenuBuilder.AddMenuEntry(
@@ -2300,7 +2010,7 @@ TSharedRef<IPropertyTypeCustomization> FPropertyBagDetails::MakeInstance()
 void FPropertyBagDetails::CustomizeHeader(TSharedRef<IPropertyHandle> StructPropertyHandle, FDetailWidgetRow& HeaderRow, IPropertyTypeCustomizationUtils& StructCustomizationUtils)
 {
 	PropUtils = StructCustomizationUtils.GetPropertyUtilities();
-	
+
 	StructProperty = StructPropertyHandle;
 	check(StructProperty);
 
@@ -2428,7 +2138,7 @@ TSharedPtr<SWidget> FPropertyBagDetails::MakeAddPropertyWidget(TSharedPtr<IPrope
 					{
 						PropertyDescs.Emplace(NewName, DefaultType);
 					});
-					
+
 				return FReply::Handled();
 
 			})
diff --git a/Engine/Source/Editor/StructUtilsEditor/Private/SInstancedStructPicker.cpp b/Engine/Source/Editor/StructUtilsEditor/Private/SInstancedStructPicker.cpp
--- a/Engine/Source/Editor/StructUtilsEditor/Private/SInstancedStructPicker.cpp
+++ b/Engine/Source/Editor/StructUtilsEditor/Private/SInstancedStructPicker.cpp
@@ -1,14 +1,18 @@
 // Copyright Epic Games, Inc. All Rights Reserved.
 
 #include "SInstancedStructPicker.h"
+
 #include "DetailLayoutBuilder.h"
 #include "Editor.h"
 #include "IPropertyUtilities.h"
 #include "PropertyCustomizationHelpers.h"
-#include "Modules/ModuleManager.h"
 #include "PropertyHandle.h"
 #include "ScopedTransaction.h"
+#include "StructUtilsEditorUtils.h"
+#include "StructUtilsMetadata.h"
+#include "Modules/ModuleManager.h"
 #include "StructUtils/InstancedStruct.h"
+#include "StructUtils/InstancedStructBaseStructQueryParams.h"
 #include "StructUtils/UserDefinedStruct.h"
 #include "Styling/SlateIconFinder.h"
 #include "Widgets/Images/SImage.h"
@@ -88,7 +92,7 @@ bool FInstancedStructFilter::IsStructAllowed(const FStructViewerInitializationOp
 	{
 		return false;
 	}
-	
+
 	if (InStruct->IsA<UUserDefinedStruct>())
 	{
 		return bAllowUserDefinedStructs;
@@ -129,9 +133,6 @@ bool FInstancedStructFilter::IsUnloadedStructAllowed(const FStructViewerInitiali
 void SInstancedStructPicker::Construct(const FArguments& InArgs, TSharedPtr<IPropertyHandle> InStructProperty,
 	TSharedPtr<IPropertyUtilities> InPropertyUtils)
 {
-	static const FName NAME_BaseStruct = "BaseStruct";
-	static const FName NAME_StructTypeConst = "StructTypeConst";
-
 	OnStructPicked = InArgs._OnStructPicked;
 	AllowedStructs = InArgs._AllowedStructs;
 	DisallowedStructs = InArgs._DisallowedStructs;
@@ -143,23 +144,7 @@ void SInstancedStructPicker::Construct(const FArguments& InArgs, TSharedPtr<IPro
 		return;
 	}
 
-	BaseScriptStruct = nullptr;
-	{
-		const FString& BaseStructName = StructProperty->GetMetaData(NAME_BaseStruct);
-		if (!BaseStructName.IsEmpty())
-		{
-			if (UScriptStruct* Struct = UClass::TryFindTypeSlow<UScriptStruct>(BaseStructName))
-			{
-				BaseScriptStruct = Struct;
-			}
-			else
-			{
-				BaseScriptStruct = LoadObject<UScriptStruct>(nullptr, *BaseStructName);
-			}
-		}
-	}
-
-	bCanChangeStructType = !StructProperty->HasMetaData(NAME_StructTypeConst);
+	bCanChangeStructType = !StructProperty->HasMetaData(UE::StructUtils::Metadata::StructTypeConstName);
 
 	ChildSlot
 	[
@@ -190,6 +175,61 @@ void SInstancedStructPicker::Construct(const FArguments& InArgs, TSharedPtr<IPro
 	];
 }
 
+const UScriptStruct* SInstancedStructPicker::GetBaseScriptStruct() const
+{
+	using namespace UE::StructUtils::Private;
+	
+	const FString& BaseStructFunctionName = StructProperty->GetMetaData(UE::StructUtils::Metadata::BaseStructFunctionName);
+	if (!BaseStructFunctionName.IsEmpty())
+	{
+		if (TOptional<FFindUserFunctionResult> Result = FindUserFunction(StructProperty,UE::StructUtils::Metadata::BaseStructFunctionName); Result.IsSet())
+		{
+			check(Result.GetValue().Function && Result.GetValue().Target);
+
+			DECLARE_DELEGATE_RetVal_OneParam(TObjectPtr<UScriptStruct>, FGetBaseStruct, UE::StructUtils::FInstancedStructBaseStructQueryParams);
+
+			UE::StructUtils::FInstancedStructBaseStructQueryParams Parameters{};
+
+			// Go up the hierarchy while we have properties
+			TSharedPtr<IPropertyHandle> CurrHandle = StructProperty;
+			while (CurrHandle.IsValid())
+			{
+				if (const FProperty* Property = CurrHandle->GetProperty())
+				{
+					Parameters.PropertyChainWithArrayIndex.Emplace(Property, CurrHandle->GetArrayIndex());
+				}
+				
+				CurrHandle = CurrHandle->GetParentHandle();
+			}
+
+			Algo::Reverse(Parameters.PropertyChainWithArrayIndex);
+			
+			if (const UScriptStruct* BaseStruct = FGetBaseStruct::CreateUFunction(Result.GetValue().Target, Result.GetValue().Function->GetFName()).Execute(Parameters))
+			{
+				return BaseStruct;
+			}
+		}
+	}
+
+	const FString& BaseStructName = StructProperty->GetMetaData(UE::StructUtils::Metadata::BaseStructName);
+	if (!BaseStructName.IsEmpty())
+	{
+		if (UScriptStruct* Struct = UClass::TryFindTypeSlow<UScriptStruct>(BaseStructName))
+		{
+			return Struct;
+		}
+
+		// TODO: We can't load an object at that point, this would break autoRTFM in Editor.
+		// Should probably have some kind of UI feedback to indicate that we are waiting.
+		// return LoadObject<UScriptStruct>(nullptr, *BaseStructName);
+		return nullptr;
+	}
+	else
+	{
+		return nullptr;
+	}
+}
+
 FText SInstancedStructPicker::GetDisplayValueString() const
 {
 	const UScriptStruct* CommonStruct = nullptr;
@@ -246,9 +286,11 @@ TSharedRef<SWidget> SInstancedStructPicker::GenerateStructPicker()
 	const bool bHideViewOptions = StructProperty->HasMetaData(NAME_HideViewOptions);
 	const bool bShowTreeView = StructProperty->HasMetaData(NAME_ShowTreeView);
 
+	const UScriptStruct* BaseScriptStruct = GetBaseScriptStruct();
+
 	TSharedRef<FInstancedStructFilter> StructFilter = MakeShared<FInstancedStructFilter>();
 	StructFilter->BaseStruct = BaseScriptStruct;
-	StructFilter->bAllowUserDefinedStructs = BaseScriptStruct.IsExplicitlyNull(); // Only allow user defined structs when BaseStruct is not set.
+	StructFilter->bAllowUserDefinedStructs = BaseScriptStruct == nullptr; // Only allow user defined structs when BaseStruct is not set.
 	StructFilter->bAllowBaseStruct = !bExcludeBaseStruct;
 
 	if (GEditor && StructProperty)
@@ -268,7 +310,7 @@ TSharedRef<SWidget> SInstancedStructPicker::GenerateStructPicker()
 		{
 			return InStruct;
 		};
-		
+
 		Algo::Transform(PropertyCustomizationHelpers::GetStructsFromMetadataString(StructProperty->GetMetaData("AllowedClasses")), StructFilter->AllowedStructs, SoftPointerTransform);
 		Algo::Transform(PropertyCustomizationHelpers::GetStructsFromMetadataString(StructProperty->GetMetaData("DisallowedClasses")), StructFilter->DisallowedStructs, SoftPointerTransform);
 
@@ -291,7 +333,7 @@ TSharedRef<SWidget> SInstancedStructPicker::GenerateStructPicker()
 						StructFilter->AllowedStructs.Append(FGetAllowedClasses::CreateUFunction(OwningObject, GetAllowedClassesFunction->GetFName()).Execute());
 					}
 				}
-		
+
 				const FString GetDisallowedClassesFunctionName = StructProperty->GetMetaData("GetDisallowedClasses");
 				if (!GetDisallowedClassesFunctionName.IsEmpty())
 				{
diff --git a/Engine/Source/Editor/StructUtilsEditor/Private/StructUtilsEditorUtils.cpp b/Engine/Source/Editor/StructUtilsEditor/Private/StructUtilsEditorUtils.cpp
new file mode 100644
--- /dev/null
+++ b/Engine/Source/Editor/StructUtilsEditor/Private/StructUtilsEditorUtils.cpp
@@ -0,0 +1,62 @@
+﻿// Copyright Epic Games, Inc. All Rights Reserved.
+
+#include "StructUtilsEditorUtils.h"
+
+namespace UE::StructUtils
+{
+
+namespace Private
+{
+
+TOptional<FFindUserFunctionResult> FindUserFunction(const TSharedPtr<IPropertyHandle>& InStructProperty, FName InFuncMetadataName)
+{
+	FProperty* MetadataProperty = InStructProperty->GetMetaDataProperty();
+
+	FFindUserFunctionResult Result{};
+
+	if (!MetadataProperty || !MetadataProperty->HasMetaData(InFuncMetadataName))
+	{
+		return {};
+	}
+
+	FString FunctionName = MetadataProperty->GetMetaData(InFuncMetadataName);
+	if (FunctionName.IsEmpty())
+	{
+		return {};
+	}
+
+	TArray<UObject*> OutObjects;
+	InStructProperty->GetOuterObjects(OutObjects);
+
+	// Check for external function references, taken from GetOptions
+	if (FunctionName.Contains(TEXT(".")))
+	{
+		Result.Function = FindObject<UFunction>(nullptr, *FunctionName, EFindObjectFlags::ExactClass);
+
+		if (ensureMsgf(Result.Function && Result.Function->HasAnyFunctionFlags(EFunctionFlags::FUNC_Static), TEXT("[%s] Didn't find function %s or expected it to be static"), *InFuncMetadataName.ToString(), *FunctionName))
+		{
+			UObject* GetOptionsCDO = Result.Function->GetOuterUClass()->GetDefaultObject();
+			Result.Target = GetOptionsCDO;
+		}
+		else
+		{
+			return {};
+		}
+	}
+	else if (OutObjects.Num() > 0)
+	{
+		Result.Target = OutObjects[0];
+		Result.Function = Result.Target->GetClass() ? Result.Target->GetClass()->FindFunctionByName(*FunctionName) : nullptr;
+	}
+
+	// Only support native functions
+	if (!ensureMsgf(Result.Function && Result.Function->IsNative(), TEXT("[%s] Didn't find function %s or expected it to be native"), *InFuncMetadataName.ToString(), *FunctionName))
+	{
+		return {};
+	}
+
+	return (Result.Target != nullptr && Result.Function != nullptr) ? Result : TOptional<FFindUserFunctionResult>{};
+}
+} // UE::StructUtils::Private
+
+} // UE::StructUtils
diff --git a/Engine/Source/Editor/StructUtilsEditor/Private/StructUtilsEditorUtils.h b/Engine/Source/Editor/StructUtilsEditor/Private/StructUtilsEditorUtils.h
new file mode 100644
--- /dev/null
+++ b/Engine/Source/Editor/StructUtilsEditor/Private/StructUtilsEditorUtils.h
@@ -0,0 +1,38 @@
+﻿// Copyright Epic Games, Inc. All Rights Reserved.
+
+#pragma once
+
+#include "PropertyHandle.h"
+#include "Templates/SharedPointer.h"
+
+class UFunction;
+class UObject;
+
+namespace UE::StructUtils
+{
+namespace Private
+{
+	
+struct FFindUserFunctionResult
+{
+	UFunction* Function = nullptr;
+	UObject* Target = nullptr;
+};
+	
+TOptional<FFindUserFunctionResult> FindUserFunction(const TSharedPtr<IPropertyHandle>& InStructProperty, FName InFuncMetadataName);
+
+/** @return true property handle holds struct property of type T.  */
+template<typename T> requires TModels_V<CStaticStructProvider, T>
+bool IsScriptStruct(const TSharedPtr<IPropertyHandle>& PropertyHandle)
+{
+	if (!PropertyHandle)
+	{
+		return false;
+	}
+
+	const FStructProperty* StructProperty = CastField<const FStructProperty>(PropertyHandle->GetProperty());
+	return StructProperty && StructProperty->Struct->IsA(TBaseStructure<T>::Get()->GetClass());
+}
+	
+} // UE::StructUtils::Private
+} // UE::StructUtils
diff --git a/Engine/Source/Editor/StructUtilsEditor/Private/StructUtilsMetadata.cpp b/Engine/Source/Editor/StructUtilsEditor/Private/StructUtilsMetadata.cpp
--- a/Engine/Source/Editor/StructUtilsEditor/Private/StructUtilsMetadata.cpp
+++ b/Engine/Source/Editor/StructUtilsEditor/Private/StructUtilsMetadata.cpp
@@ -22,6 +22,10 @@ namespace UE::StructUtils::Metadata
 	const FLazyName CanRemovePropertyName = "CanRemoveProperty";
 	const FLazyName ChildRowFeaturesName = "ChildRowFeatures";
 
+	const FLazyName BaseStructName = "BaseStruct";
+	const FLazyName BaseStructFunctionName = "BaseStructFunction";
+	const FLazyName StructTypeConstName = "StructTypeConst";
+
 	bool AreCategoriesEnabled(const FPropertyBagPropertyDesc& Desc)
 	{
 		return Desc.HasMetaData(EnableCategoriesName);
diff --git a/Engine/Source/Editor/StructUtilsEditor/Public/SInstancedStructPicker.h b/Engine/Source/Editor/StructUtilsEditor/Public/SInstancedStructPicker.h
--- a/Engine/Source/Editor/StructUtilsEditor/Public/SInstancedStructPicker.h
+++ b/Engine/Source/Editor/StructUtilsEditor/Public/SInstancedStructPicker.h
@@ -52,10 +52,7 @@ public:
 
 	UE_API void Construct(const FArguments& InArgs, TSharedPtr<IPropertyHandle> InStructProperty, TSharedPtr<IPropertyUtilities> InPropertyUtils);
 
-	const UScriptStruct* GetBaseScriptStruct() const
-	{
-		return BaseScriptStruct.Get();
-	}
+	const UScriptStruct* GetBaseScriptStruct() const;
 
 	bool CanChangeStructType() const
 	{
@@ -69,9 +66,6 @@ private:
 	TSharedPtr<IPropertyHandle> StructProperty;
 	TSharedPtr<IPropertyUtilities> PropUtils;
 
-	/** The base struct that we're allowing to be picked (controlled by the "BaseStruct" meta-data) */
-	TWeakObjectPtr<UScriptStruct> BaseScriptStruct = nullptr;
-
 	/** Can we change the struct type at all? (controlled by the "StructTypeConst" meta-data) */
 	bool bCanChangeStructType = true;
 
diff --git a/Engine/Source/Editor/StructUtilsEditor/Public/StructUtilsMetadata.h b/Engine/Source/Editor/StructUtilsEditor/Public/StructUtilsMetadata.h
--- a/Engine/Source/Editor/StructUtilsEditor/Public/StructUtilsMetadata.h
+++ b/Engine/Source/Editor/StructUtilsEditor/Public/StructUtilsMetadata.h
@@ -42,6 +42,19 @@ namespace UE::StructUtils::Metadata
 	/// [Property Metadata] CategoryName: The name of the category (grouping of properties).
 	extern STRUCTUTILSEDITOR_API const FLazyName CategoryName;
 
+
+	/// Property metadata specifiers for FInstancedStruct/TInstancedStruct
+	/// [PropertyMetadata] BaseStruct: Type name of the base class for the instanced struct
+	extern STRUCTUTILSEDITOR_API const FLazyName BaseStructName;
+
+	/// [PropertyMetadata] BaseStructFunction: Name of a UFunction with signature UScriptStruct*(UE::StructUtils::FInstancedStructBaseStructQueryParams). Returns the type of the base struct for the instanced struct.
+	/// (cf. InstancedStructBaseStructQueryParams.h for an example of function.)
+	/// @note When return value is invalid, BaseStruct will be used if any.
+	extern STRUCTUTILSEDITOR_API const FLazyName BaseStructFunctionName;
+
+	/// [PropertyMetadata] StructTypeConst: Determine whether the type of the instanced struct can be changed
+	extern STRUCTUTILSEDITOR_API const FLazyName StructTypeConstName;
+
 	/** If organizing by category is enabled for this property. */
 	STRUCTUTILSEDITOR_API bool AreCategoriesEnabled(const FPropertyBagPropertyDesc& Desc);
 	/** Enable category organization for this property. */
diff --git a/Engine/Source/Runtime/CoreUObject/Public/StructUtils/InstancedStructBaseStructQueryParams.h b/Engine/Source/Runtime/CoreUObject/Public/StructUtils/InstancedStructBaseStructQueryParams.h
new file mode 100644
--- /dev/null
+++ b/Engine/Source/Runtime/CoreUObject/Public/StructUtils/InstancedStructBaseStructQueryParams.h
@@ -0,0 +1,112 @@
+﻿// Copyright Epic Games, Inc. All Rights Reserved.
+
+#pragma once
+
+#include "InstancedStructBaseStructQueryParams.generated.h"
+
+/**
+ * Wrapper structure around an array of Properties to be used as a parameter of a UFunction for querying the base struct of the struct picker in UI.
+ * It will be filled with the chain of properties from the owner object to the customizable property,
+ * useful if the customizable InstancedStruct is nested in other structs.
+ * It also have an optional array index if the struct is nested in an array.
+ * TODO: Add support for struct nested in Sets and Maps.
+ *
+ * Example:
+ * 
+ * USTRUCT(BlueprintType)
+ * struct FBaseStruct
+ * {
+ * 	UPROPERTY(EditAnywhere, BlueprintReadWrite, Category = "Custom")
+ * 	int IntProperty = 0;
+ * };
+ * 
+ * USTRUCT(BlueprintType)
+ * struct FChildStruct1 : public FBaseStruct
+ * {
+ * 	UPROPERTY(EditAnywhere, BlueprintReadWrite, Category = "Custom")
+ * 	float FloatProperty = 0;
+ * };
+ * 
+ * USTRUCT(BlueprintType)
+ * struct FChildStruct2 : public FBaseStruct
+ * {
+ * 	UPROPERTY(EditAnywhere, BlueprintReadWrite, Category = "Custom")
+ * 	FString StringProperty;
+ * };
+ * 
+ * USTRUCT(BlueprintType)
+ * struct FNestedStruct
+ * {
+ * 	UPROPERTY(EditAnywhere, BlueprintReadWrite, Category = "Custom")
+ * 	int StructToChoose = 0;
+ * 
+ * 	UPROPERTY(EditAnywhere, BlueprintReadWrite, Category = "Custom", meta = (BaseStructFunction = "GetSettingScriptStruct"))
+ * 	TInstancedStruct<FBaseStruct> MyInstancedStruct;
+ * };
+ * 
+ * UCLASS()
+ * class UCustomObject : public UObject
+ * {
+ * public:
+ * 	UPROPERTY(EditAnywhere, BlueprintReadWrite, Category = "Custom")
+ * 	TArray<FNestedStruct> NestedStructs;
+ * 
+ * 	UPROPERTY(EditAnywhere, BlueprintReadWrite, Category = "Custom")
+ * 	FNestedStruct NestedStruct2;
+ * 
+ * #if WITH_EDITOR
+ * 	UFUNCTION()
+ * 	UScriptStruct* GetSettingScriptStruct(UE::StructUtils::FInstancedStructBaseStructQueryParams Params) const
+ * 	{
+ * 		int32 Index = -1;
+ * 		if (Params.PropertyChainWithArrayIndex.Num() == 2)
+ * 		{
+ * 			if (Params.PropertyChainWithArrayIndex[0].Property->GetFName() == GET_MEMBER_NAME_CHECKED(UCustomObject, NestedStruct2))
+ * 			{
+ * 				Index = NestedStruct2.StructToChoose;
+ * 			}
+ * 		}
+ * 		else if (Params.PropertyChainWithArrayIndex.Num() == 3)
+ * 		{
+ * 			if (Params.PropertyChainWithArrayIndex[1].Property->GetFName() == GET_MEMBER_NAME_CHECKED(UCustomObject, NestedStructs))
+ * 			{
+ * 				const int32 ArrayIndex = Params.PropertyChainWithArrayIndex[1].ArrayIndex;
+ * 				if (ensure(NestedStructs.IsValidIndex(ArrayIndex)))
+ * 				{
+ * 					Index = NestedStructs[ArrayIndex].StructToChoose;
+ * 				}
+ * 			}
+ * 		}
+ * 
+ * 		if (Index != -1)
+ * 		{
+ * 			return Index == 0 ? FChildStruct1::StaticStruct() : FChildStruct2::StaticStruct();
+ * 		}
+ * 		else
+ * 		{
+ * 			return FBaseStruct::StaticStruct();
+ * 		}
+ * 	}
+ * #endif // WITH_EDITOR
+ * };
+ */
+
+namespace UE::StructUtils
+{
+USTRUCT(BlueprintInternalUseOnly)
+struct FInstancedStructBaseStructQueryParams
+{
+	GENERATED_BODY()
+
+#if WITH_EDITOR
+	struct FParam
+	{
+		const FProperty* Property = nullptr;
+		int32 ArrayIndex = INDEX_NONE;
+	};
+	
+	TArray<FParam> PropertyChainWithArrayIndex;
+#endif // WITH_EDITOR
+};
+}
+
