Property Bag Details: removed check on property and return early if invalid as it is not guaranteed valid even if the handle is

diff --git a/Engine/Source/Editor/StructUtilsEditor/Private/Customizations/PropertyBagDetails.cpp b/Engine/Source/Editor/StructUtilsEditor/Private/Customizations/PropertyBagDetails.cpp
--- a/Engine/Source/Editor/StructUtilsEditor/Private/Customizations/PropertyBagDetails.cpp
+++ b/Engine/Source/Editor/StructUtilsEditor/Private/Customizations/PropertyBagDetails.cpp
@@ -1762,15 +1762,18 @@ FPropertyBagInstanceDataDetails::EPropertyOverrideState FPropertyBagInstanceData
 {
 	if (!ChildPropertyHandle)
 	{
-		return EPropertyOverrideState::Undetermined;;
+		return EPropertyOverrideState::Undetermined;
+	}
+
+	const FProperty* Property = ChildPropertyHandle->GetProperty();
+	if (!Property)
+	{
+		return EPropertyOverrideState::Undetermined;
 	}
 
 	int32 NumValues = 0;
 	int32 NumOverrides = 0;
 
-	const FProperty* Property = ChildPropertyHandle->GetProperty();
-	check(Property);
-
 	EnumeratePropertyBags(BagStructProperty,
 		[Property, &NumValues, &NumOverrides]
 		(const FInstancedPropertyBag& DefaultPropertyBag, const FInstancedPropertyBag& PropertyBag, const IPropertyBagOverrideProvider& OverrideProvider)
