[StructUtils] Adjusted desired size of property bag details property input text box for FString and FName

#rnx

diff --git a/Engine/Source/Editor/StructUtilsEditor/Private/Customizations/PropertyBagDetails.cpp b/Engine/Source/Editor/StructUtilsEditor/Private/Customizations/PropertyBagDetails.cpp
--- a/Engine/Source/Editor/StructUtilsEditor/Private/Customizations/PropertyBagDetails.cpp
+++ b/Engine/Source/Editor/StructUtilsEditor/Private/Customizations/PropertyBagDetails.cpp
@@ -18,6 +18,7 @@
 #include "StructUtils/UserDefinedStruct.h"
 #include "Templates/ValueOrError.h"
 #include "UObject/EnumProperty.h"
+#include "UObject/TextProperty.h"
 
 #include "Widgets/Images/SImage.h"
 #include "Widgets/Input/SButton.h"
@@ -676,6 +677,26 @@ void PinInfoChanged(TSharedPtr<IPropertyHandle> ChildPropertyHandle, const TShar
 			}
 		});
 }
+
+FOptionalSize GetDesiredInputWidgetSizeByType(const TSharedPtr<IPropertyHandle>& ChildPropertyHandle)
+{
+	static constexpr float StringPropertyDesiredSize = 120.f;
+	static constexpr float NamePropertyDesiredSize = 50.f;
+
+	if (const FProperty* Property = ChildPropertyHandle ? ChildPropertyHandle->GetProperty() : nullptr)
+	{
+		if (Property->IsA(FStrProperty::StaticClass()))
+		{
+			return StringPropertyDesiredSize;
+		}
+		else if (Property->IsA(FNameProperty::StaticClass()))
+		{
+			return NamePropertyDesiredSize;
+		}
+	}
+
+	return {};
+}
 } // UE::StructUtils::Private
 
 TSharedRef<SWidget> CreateTypeSelectionWidget(TSharedPtr<IPropertyHandle> ChildPropertyHandle, const TSharedPtr<IPropertyHandle>& InBagStructProperty, const TSharedPtr<IPropertyUtilities>& InPropUtils, SPinTypeSelector::ESelectorType SelectorType, const bool bAllowContainers)
@@ -995,11 +1016,19 @@ void FPropertyBagInstanceDataDetails::OnChildRowAdded(IDetailPropertyRow& ChildR
 	TSharedPtr<SWidget> NameWidget;
 	TSharedPtr<SWidget> PropertyValueWidget;
 	FDetailWidgetRow DetailWidgetRow;
-	ChildRow.GetDefaultWidgets(NameWidget, PropertyValueWidget, DetailWidgetRow);
 
 	TSharedPtr<IPropertyHandle> ChildPropertyHandle = ChildRow.GetPropertyHandle();
 	check(ChildPropertyHandle);
 
+	ChildRow.GetDefaultWidgets(NameWidget, PropertyValueWidget, DetailWidgetRow);
+	// Wrap the widget in a box of the desired size, based on the property type.
+	PropertyValueWidget = SNew(SBox)
+		.MinDesiredWidth(Private::GetDesiredInputWidgetSizeByType(ChildPropertyHandle))
+		.Content()
+		[
+			PropertyValueWidget.ToSharedRef()
+		];
+
 	TWeakPtr<FPropertyBagInstanceDataDetails> WeakSelf = SharedThis<FPropertyBagInstanceDataDetails>(this);
 
 	const UPropertyBag* BagStruct = BagStructProperty ? Private::GetCommonBagStruct(BagStructProperty) : nullptr;
