Moved parent/child check for instanced struct picker above asset reference check.
This resolves an issue with domains not being handled correctly leading to an unnecessary ensure being triggered

diff --git a/Engine/Source/Editor/StructUtilsEditor/Private/SInstancedStructPicker.cpp b/Engine/Source/Editor/StructUtilsEditor/Private/SInstancedStructPicker.cpp
--- a/Engine/Source/Editor/StructUtilsEditor/Private/SInstancedStructPicker.cpp
+++ b/Engine/Source/Editor/StructUtilsEditor/Private/SInstancedStructPicker.cpp
@@ -103,6 +103,12 @@ bool FInstancedStructFilter::IsStructAllowed(const FStructViewerInitializationOp
 		return bAllowBaseStruct;
 	}
 
+	// If we have a base struct, make sure the candidate struct is a child
+	if (BaseStruct.IsValid() && !InStruct->IsChildOf(BaseStruct.Get()))
+	{
+		return false;
+	}
+
 	static const FName NAME_HiddenMetaTag = "Hidden";
 	if (InStruct->HasMetaData(NAME_HiddenMetaTag))
 	{
@@ -116,10 +122,8 @@ bool FInstancedStructFilter::IsStructAllowed(const FStructViewerInitializationOp
 			return false;
 		}
 	}
-
-	// Query the native struct to see if it has the correct parent type (if any)
-	const UScriptStruct* Struct = BaseStruct.Get();
-	return !Struct || InStruct->IsChildOf(Struct);
+	
+	return true;
 }
 
 bool FInstancedStructFilter::IsUnloadedStructAllowed(const FStructViewerInitializationOptions& InInitOptions, const FSoftObjectPath& InStructPath, TSharedRef<FStructViewerFilterFuncs> InFilterFuncs)
