[StructUtils] Fixed crash when losing focus on setting metadata specifiers for PropertyBag properties.

+ Also clearing the UI text box should now properly remove the specifier from the metadata.

#rnx

diff --git a/Engine/Source/Editor/StructUtilsEditor/Private/Customizations/PropertyBagDetails.cpp b/Engine/Source/Editor/StructUtilsEditor/Private/Customizations/PropertyBagDetails.cpp
--- a/Engine/Source/Editor/StructUtilsEditor/Private/Customizations/PropertyBagDetails.cpp
+++ b/Engine/Source/Editor/StructUtilsEditor/Private/Customizations/PropertyBagDetails.cpp
@@ -401,7 +401,7 @@ bool IsUniqueName(const FName NewName, const FName OldName, const TSharedPtr<IPr
 template<typename TFunc>
 void ApplyChangesToPropertyDescs(const FText& SessionName, const TSharedPtr<IPropertyHandle>& StructProperty, const TSharedPtr<IPropertyUtilities>& PropUtils, TFunc&& Function)
 {
-	if (!StructProperty || !PropUtils)
+	if (!StructProperty || !StructProperty.Get()->IsValidHandle() || !PropUtils)
 	{
 		return;
 	}
@@ -1550,7 +1550,7 @@ void FPropertyBagInstanceDataDetails::OnChildRowAdded(IDetailPropertyRow& ChildR
 							? FOnVerifyTextChanged()
 							: FOnVerifyTextChanged::CreateLambda([](const FText& InText, FText& OutMessage) -> bool
 							{
-								if (!InText.IsNumeric())
+								if (!InText.IsEmpty() && !InText.IsNumeric())
 								{
 									OutMessage = LOCTEXT("DropDownMenuSpecifierNumericValueExpected", "Value is not numeric");
 									return false;
@@ -1573,7 +1573,15 @@ void FPropertyBagInstanceDataDetails::OnChildRowAdded(IDetailPropertyRow& ChildR
 										PropUtils,
 										[SpecifierName, &CommittedText](FPropertyBagPropertyDesc& Desc)
 										{
-											Desc.SetMetaData(SpecifierName, CommittedText.ToString());
+											// If empty, remove the specifier
+											if (CommittedText.IsEmpty())
+											{
+												Desc.RemoveMetadata(SpecifierName);
+											}
+											else
+											{
+												Desc.SetMetaData(SpecifierName, CommittedText.ToString());
+											}
 										});
 								}
 							}));
diff --git a/Engine/Source/Runtime/CoreUObject/Private/StructUtils/PropertyBag.cpp b/Engine/Source/Runtime/CoreUObject/Private/StructUtils/PropertyBag.cpp
--- a/Engine/Source/Runtime/CoreUObject/Private/StructUtils/PropertyBag.cpp
+++ b/Engine/Source/Runtime/CoreUObject/Private/StructUtils/PropertyBag.cpp
@@ -1911,6 +1911,14 @@ FString FPropertyBagPropertyDesc::GetMetaData(FName SpecifierName) const
 		return {};
 	}
 }
+
+void FPropertyBagPropertyDesc::RemoveMetadata(FName SpecifierName)
+{
+	MetaData.RemoveAll([SpecifierName](const FPropertyBagPropertyDescMetaData& MetaData)
+	{
+		return MetaData.Key == SpecifierName;
+	});
+}
 #endif // WITH_EDITOR
 
 //----------------------------------------------------------------//
diff --git a/Engine/Source/Runtime/CoreUObject/Public/StructUtils/PropertyBag.h b/Engine/Source/Runtime/CoreUObject/Public/StructUtils/PropertyBag.h
--- a/Engine/Source/Runtime/CoreUObject/Public/StructUtils/PropertyBag.h
+++ b/Engine/Source/Runtime/CoreUObject/Public/StructUtils/PropertyBag.h
@@ -318,6 +318,9 @@ struct FPropertyBagPropertyDesc
 
 	/** @return retrieves the metadata specifier value */
 	UE_API FString GetMetaData(FName SpecifierName) const;
+
+	/** Remove the metadata specifier value from the property. */
+	UE_API void RemoveMetadata(FName SpecifierName);
 #endif // WITH_EDITOR
 
 	/** Pointer to object that defines the Enum, Struct, or Class. */
