[StructUtils] Add extension point which passes ChildHandle and ChildrenBuilder
#pr10398

diff --git a/Engine/Source/Editor/StructUtilsEditor/Private/InstancedStructDetails.cpp b/Engine/Source/Editor/StructUtilsEditor/Private/InstancedStructDetails.cpp
--- a/Engine/Source/Editor/StructUtilsEditor/Private/InstancedStructDetails.cpp
+++ b/Engine/Source/Editor/StructUtilsEditor/Private/InstancedStructDetails.cpp
@@ -351,8 +351,7 @@ void FInstancedStructDataDetails::GenerateChildContent(IDetailChildrenBuilder& C
 	CachedInstanceTypes = GetInstanceTypes();
 }
 
-void FInstancedStructDataDetails::AddChildRows(IDetailChildrenBuilder& ChildBuilder,
-	const TArray<TSharedPtr<IPropertyHandle>>& ChildProperties)
+void FInstancedStructDataDetails::AddChildRows(IDetailChildrenBuilder& ChildBuilder, const TArray<TSharedPtr<IPropertyHandle>>& ChildProperties)
 {
 	// Properties may have Category metadata. If that's the case, they should be added under groups.
 	TMap<TSharedPtr<IPropertyHandle>, IDetailGroup*> PropertyToGroup;
@@ -360,18 +359,24 @@ void FInstancedStructDataDetails::AddChildRows(IDetailChildrenBuilder& ChildBuil
 
 	for (TSharedPtr<IPropertyHandle> ChildHandle : ChildProperties)
 	{
-		// If the property has a group, add it under the group. Otherwise, just add it normally via the builder.
-		IDetailGroup** PropertyGroup = PropertyToGroup.Find(ChildHandle);
-		if (PropertyGroup && *PropertyGroup)
-		{
-			IDetailPropertyRow& Row = (*PropertyGroup)->AddPropertyRow(ChildHandle.ToSharedRef());
-			OnChildRowAdded(Row);
-		}
-		else
-		{
-			IDetailPropertyRow& Row = ChildBuilder.AddProperty(ChildHandle.ToSharedRef());
-			OnChildRowAdded(Row);
-		}
+		IDetailGroup** PropertyGroupPtr = PropertyToGroup.Find(ChildHandle);
+		IDetailGroup* PropertyGroup = PropertyGroupPtr ? *PropertyGroupPtr : nullptr;
+		AddChildRow(ChildBuilder, ChildHandle, PropertyGroup);
+	}
+}
+
+void FInstancedStructDataDetails::AddChildRow(IDetailChildrenBuilder& ChildBuilder, TSharedPtr<IPropertyHandle> ChildHandle, IDetailGroup* PropertyGroup)
+{
+	// If the property has a group, add it under the group. Otherwise, just add it normally via the builder.
+	if (PropertyGroup)
+	{
+		IDetailPropertyRow& Row = PropertyGroup->AddPropertyRow(ChildHandle.ToSharedRef());
+		OnChildRowAdded(Row);
+	}
+	else
+	{
+		IDetailPropertyRow& Row = ChildBuilder.AddProperty(ChildHandle.ToSharedRef());
+		OnChildRowAdded(Row);
 	}
 }
 
diff --git a/Engine/Source/Editor/StructUtilsEditor/Public/InstancedStructDetails.h b/Engine/Source/Editor/StructUtilsEditor/Public/InstancedStructDetails.h
--- a/Engine/Source/Editor/StructUtilsEditor/Public/InstancedStructDetails.h
+++ b/Engine/Source/Editor/StructUtilsEditor/Public/InstancedStructDetails.h
@@ -83,7 +83,10 @@ public:
 
 protected:
 	UE_API static TArray<FString> GetPropertyCategories(TSharedPtr<IPropertyHandle> PropertyHandle);
+	/** Add all the rows for a child properties. */
 	UE_API virtual void AddChildRows(IDetailChildrenBuilder& ChildBuilder, const TArray<TSharedPtr<IPropertyHandle>>& ChildProperties);
+	/** Add all the rows for a child's property. */
+	UE_API virtual void AddChildRow(IDetailChildrenBuilder& ChildBuilder, TSharedPtr<IPropertyHandle> ChildHandle, IDetailGroup* PropertyGroup);
 
 private:
 	UE_API void OnUserDefinedStructReinstancedHandle(const class UUserDefinedStruct& Struct);
