[StructUtils] Fix typo in log.

#rnx

diff --git a/Engine/Source/Editor/StructUtilsEditor/Public/PropertyBagDetails.h b/Engine/Source/Editor/StructUtilsEditor/Public/PropertyBagDetails.h
--- a/Engine/Source/Editor/StructUtilsEditor/Public/PropertyBagDetails.h
+++ b/Engine/Source/Editor/StructUtilsEditor/Public/PropertyBagDetails.h
@@ -189,7 +189,7 @@ public:
 	{
 		// ParentPropertyHandle would be the property that holds the property overrides. 
 		//		ParentPropertyHandle->NotifyPreChange();
-		checkf(false, TEXT("PreChangeOverrides() is expecgted to be implemented when HasPropertyOverrides() returns true."));
+		checkf(false, TEXT("PreChangeOverrides() is expected to be implemented when HasPropertyOverrides() returns true."));
 	}
 
 	/** Called after property override is changed. */
@@ -198,7 +198,7 @@ public:
 		// ParentPropertyHandle is be the property that holds the property overrides. 
 		//		ParentPropertyHandle->NotifyPostChange(EPropertyChangeType::ValueSet);
 		//		ParentPropertyHandle->NotifyFinishedChangingProperties();
-		checkf(false, TEXT("PostChangeOverrides() is expecgted to be implemented when HasPropertyOverrides() returns true."));
+		checkf(false, TEXT("PostChangeOverrides() is expected to be implemented when HasPropertyOverrides() returns true."));
 	}
 
 	/**
