PropertyBag: Fix up Access Specifier toggle feature.
- Now the eye is open when the property is public, closed otherwise.
- Clicking the eye will now result in a proper transaction.

diff --git a/Engine/Source/Editor/StructUtilsEditor/Private/Customizations/PropertyBagDetails.cpp b/Engine/Source/Editor/StructUtilsEditor/Private/Customizations/PropertyBagDetails.cpp
--- a/Engine/Source/Editor/StructUtilsEditor/Private/Customizations/PropertyBagDetails.cpp
+++ b/Engine/Source/Editor/StructUtilsEditor/Private/Customizations/PropertyBagDetails.cpp
@@ -1305,27 +1305,22 @@ void FPropertyBagInstanceDataDetails::OnChildRowAdded(IDetailPropertyRow& ChildR
 					.OnClicked_Lambda([ChildPropertyHandle, BagStructProperty = BagStructProperty, PropUtils = PropUtils]()
 					{
 						FProperty* Property = ChildPropertyHandle->GetProperty();
-						Private::ApplyChangesToPropertyDescs(
+						Private::ApplyChangesToSinglePropertyDesc(
 							LOCTEXT("OnPropertyAccessSpecifierChanged", "Set access specifier."),
+							ChildPropertyHandle,
 							BagStructProperty,
 							PropUtils,
-							[Property, BagStructProperty](TArray<FPropertyBagPropertyDesc>& PropertyDescs)
+							[Property, BagStructProperty](FPropertyBagPropertyDesc& Desc)
 							{
-								if (FPropertyBagPropertyDesc* Desc = PropertyDescs.FindByPredicate([Property, BagStructProperty](const FPropertyBagPropertyDesc& Desc) { return Desc.CachedProperty == Property; }))
+								const bool bIsPrivate = !Property->HasAnyPropertyFlags(CPF_NativeAccessSpecifierPublic);
+								Desc.PropertyFlags &= ~CPF_NativeAccessSpecifiers;
+								if (bIsPrivate)
 								{
-									BagStructProperty->NotifyPreChange();
-									const bool bIsPrivate = Property->HasAnyPropertyFlags(CPF_NativeAccessSpecifierPrivate | CPF_NativeAccessSpecifierProtected);
-									Desc->PropertyFlags &= ~CPF_NativeAccessSpecifiers;
-									if (bIsPrivate)
-									{
-										Desc->PropertyFlags |= CPF_NativeAccessSpecifierPublic;
-									}
-									else
-									{
-										Desc->PropertyFlags |= CPF_NativeAccessSpecifierPrivate;
-									}
-									BagStructProperty->NotifyPostChange(EPropertyChangeType::ValueSet);
-									BagStructProperty->NotifyFinishedChangingProperties();
+									Desc.PropertyFlags |= CPF_NativeAccessSpecifierPublic;
+								}
+								else
+								{
+									Desc.PropertyFlags |= CPF_NativeAccessSpecifierPrivate;
 								}
 							});
 
@@ -1342,7 +1337,7 @@ void FPropertyBagInstanceDataDetails::OnChildRowAdded(IDetailPropertyRow& ChildR
 							if (const FProperty* Property = ChildPropertyHandle->GetProperty())
 							{
 								// For now, treat protected as private. TODO: Add toggle for protected.
-								if (Property->HasAnyPropertyFlags(CPF_NativeAccessSpecifierPrivate | CPF_NativeAccessSpecifierProtected))
+								if (Property->HasAnyPropertyFlags(CPF_NativeAccessSpecifierPublic))
 								{
 									return FAppStyle::Get().GetBrush("Icons.Visible");
 								}
