PropertyBag: Expose bag mutating methods from the details panel.
This lets systems to customize their property bags with extra features.

diff --git a/Engine/Source/Editor/StructUtilsEditor/Private/Customizations/PropertyBagDetails.cpp b/Engine/Source/Editor/StructUtilsEditor/Private/Customizations/PropertyBagDetails.cpp
--- a/Engine/Source/Editor/StructUtilsEditor/Private/Customizations/PropertyBagDetails.cpp
+++ b/Engine/Source/Editor/StructUtilsEditor/Private/Customizations/PropertyBagDetails.cpp
@@ -279,36 +279,6 @@ bool HasMissingType(const TSharedPtr<IPropertyHandle>& PropertyHandle)
 	return false;
 }
 
-/** @return property bag struct common to all edited properties. */
-const UPropertyBag* GetCommonBagStruct(TSharedPtr<IPropertyHandle> StructProperty)
-{
-	const UPropertyBag* CommonBagStruct = nullptr;
-
-	if (ensure(IsScriptStruct<FInstancedPropertyBag>(StructProperty)))
-	{
-		StructProperty->EnumerateConstRawData([&CommonBagStruct](const void* RawData, const int32 /*DataIndex*/, const int32 /*NumDatas*/)
-		{
-			if (RawData)
-			{
-				const FInstancedPropertyBag* Bag = static_cast<const FInstancedPropertyBag*>(RawData);
-
-				const UPropertyBag* BagStruct = Bag->GetPropertyBagStruct();
-				if (CommonBagStruct && CommonBagStruct != BagStruct)
-				{
-					// Multiple struct types on the sources - show nothing set
-					CommonBagStruct = nullptr;
-					return false;
-				}
-				CommonBagStruct = BagStruct;
-			}
-
-			return true;
-		});
-	}
-
-	return CommonBagStruct;
-}
-
 /** @return property descriptors of the property bag struct common to all edited properties. */
 TArray<FPropertyBagPropertyDesc> GetCommonPropertyDescs(const TSharedPtr<IPropertyHandle>& StructProperty)
 {
@@ -398,54 +368,6 @@ bool IsUniqueName(const FName NewName, const FName OldName, const TSharedPtr<IPr
 	return !bFound;
 }
 
-template<typename TFunc>
-void ApplyChangesToPropertyDescs(const FText& SessionName, const TSharedPtr<IPropertyHandle>& StructProperty, const TSharedPtr<IPropertyUtilities>& PropUtils, TFunc&& Function)
-{
-	if (!StructProperty || !PropUtils)
-	{
-		return;
-	}
-
-	FScopedTransaction Transaction(SessionName);
-	TArray<FPropertyBagPropertyDesc> PropertyDescs = GetCommonPropertyDescs(StructProperty);
-	StructProperty->NotifyPreChange();
-
-	Function(PropertyDescs);
-
-	SetPropertyDescs(StructProperty, PropertyDescs);
-
-	StructProperty->NotifyPostChange(EPropertyChangeType::ValueSet);
-	StructProperty->NotifyFinishedChangingProperties();
-	if (PropUtils)
-	{
-		PropUtils->ForceRefresh();
-	}
-}
-
-template <typename TFunc>
-void ApplyChangesToSinglePropertyDesc(const FText& SessionName, const TSharedPtr<IPropertyHandle> PropertyHandle, const TSharedPtr<IPropertyHandle>& StructProperty, const TSharedPtr<IPropertyUtilities>& PropUtils, TFunc Function)
-{
-	ApplyChangesToPropertyDescs(SessionName, StructProperty, PropUtils, [Function = std::move(Function), Property = PropertyHandle->GetProperty()](TArray<FPropertyBagPropertyDesc>& PropertyDescs)
-	{
-		if (FPropertyBagPropertyDesc* Desc = PropertyDescs.FindByPredicate([Property](const FPropertyBagPropertyDesc& OutDesc) { return OutDesc.CachedProperty == Property; }))
-		{
-			Function(*Desc);
-		}
-	});
-}
-
-template <typename TFunc>
-void ApplyChangesToSinglePropertyDesc(const FText& SessionName, const FPropertyBagPropertyDesc& PropertyDesc, const TSharedPtr<IPropertyHandle>& StructProperty, const TSharedPtr<IPropertyUtilities>& PropUtils, TFunc Function)
-{
-	ApplyChangesToPropertyDescs(SessionName, StructProperty, PropUtils, [Function = std::move(Function), &PropertyDesc](TArray<FPropertyBagPropertyDesc>& PropertyDescs)
-	{
-		if (FPropertyBagPropertyDesc* Desc = PropertyDescs.FindByPredicate([&PropertyDesc](const FPropertyBagPropertyDesc& OutDesc) { return OutDesc == PropertyDesc; }))
-		{
-			Function(*Desc);
-		}
-	});
-}
-
 bool CanHaveMemberVariableOfType(const FEdGraphPinType& PinType)
 {
 	if (PinType.PinCategory == UEdGraphSchema_K2::PC_Exec
@@ -663,7 +585,7 @@ void PinInfoChanged(TSharedPtr<IPropertyHandle> ChildPropertyHandle, const TShar
 		return;
 	}
 
-	Private::ApplyChangesToPropertyDescs(
+	ApplyChangesToPropertyDescs(
 		FText::Format(LOCTEXT("OnPropertyTypeChanged", "Changed property type: {0}"), ChildPropertyHandle->GetPropertyDisplayName()),
 		InBagStructProperty,
 		InPropUtils,
@@ -726,6 +648,80 @@ TSharedRef<SWidget> CreateTypeSelectionWidget(TSharedPtr<IPropertyHandle> ChildP
 		];
 }
 
+const UPropertyBag* GetCommonBagStruct(TSharedPtr<IPropertyHandle> StructProperty)
+{
+	const UPropertyBag* CommonBagStruct = nullptr;
+
+	if (ensure(Private::IsScriptStruct<FInstancedPropertyBag>(StructProperty)))
+	{
+		StructProperty->EnumerateConstRawData([&CommonBagStruct](const void* RawData, const int32 /*DataIndex*/, const int32 /*NumDatas*/)
+		{
+			if (RawData)
+			{
+				const FInstancedPropertyBag* Bag = static_cast<const FInstancedPropertyBag*>(RawData);
+
+				const UPropertyBag* BagStruct = Bag->GetPropertyBagStruct();
+				if (CommonBagStruct && CommonBagStruct != BagStruct)
+				{
+					// Multiple struct types on the sources - show nothing set
+					CommonBagStruct = nullptr;
+					return false;
+				}
+				CommonBagStruct = BagStruct;
+			}
+
+			return true;
+		});
+	}
+
+	return CommonBagStruct;
+}
+
+void ApplyChangesToPropertyDescs(const FText& SessionName, const TSharedPtr<IPropertyHandle>& StructProperty, const TSharedPtr<IPropertyUtilities>& PropUtils, TFunctionRef<void(TArray<FPropertyBagPropertyDesc>&)> Function)
+{
+	if (!StructProperty || !PropUtils)
+	{
+		return;
+	}
+
+	FScopedTransaction Transaction(SessionName);
+	TArray<FPropertyBagPropertyDesc> PropertyDescs = Private::GetCommonPropertyDescs(StructProperty);
+	StructProperty->NotifyPreChange();
+
+	Function(PropertyDescs);
+
+	Private::SetPropertyDescs(StructProperty, PropertyDescs);
+
+	StructProperty->NotifyPostChange(EPropertyChangeType::ValueSet);
+	StructProperty->NotifyFinishedChangingProperties();
+	if (PropUtils)
+	{
+		PropUtils->ForceRefresh();
+	}
+}
+
+void ApplyChangesToSinglePropertyDesc(const FText& SessionName, const TSharedPtr<IPropertyHandle> PropertyHandle, const TSharedPtr<IPropertyHandle>& StructProperty, const TSharedPtr<IPropertyUtilities>& PropUtils, TFunctionRef<void(FPropertyBagPropertyDesc&)> Function)
+{
+	ApplyChangesToPropertyDescs(SessionName, StructProperty, PropUtils, [Function = std::move(Function), Property = PropertyHandle->GetProperty()](TArray<FPropertyBagPropertyDesc>& PropertyDescs)
+	{
+		if (FPropertyBagPropertyDesc* Desc = PropertyDescs.FindByPredicate([Property](const FPropertyBagPropertyDesc& OutDesc) { return OutDesc.CachedProperty == Property; }))
+		{
+			Function(*Desc);
+		}
+	});
+}
+
+void ApplyChangesToSinglePropertyDesc(const FText& SessionName, const FPropertyBagPropertyDesc& PropertyDesc, const TSharedPtr<IPropertyHandle>& StructProperty, const TSharedPtr<IPropertyUtilities>& PropUtils, TFunctionRef<void(FPropertyBagPropertyDesc&)> Function)
+{
+	ApplyChangesToPropertyDescs(SessionName, StructProperty, PropUtils, [Function = std::move(Function), &PropertyDesc](TArray<FPropertyBagPropertyDesc>& PropertyDescs)
+	{
+		if (FPropertyBagPropertyDesc* Desc = PropertyDescs.FindByPredicate([&PropertyDesc](const FPropertyBagPropertyDesc& OutDesc) { return OutDesc == PropertyDesc; }))
+		{
+			Function(*Desc);
+		}
+	});
+}
+
 namespace Constants
 {
 	static constexpr int32 MaxCategoryLength = 70;
@@ -823,14 +819,14 @@ void FPropertyBagInstanceDataDetails::OnGroupRowAdded(IDetailGroup& GroupRow, in
 					if (ensure(DroppedPropertyDesc.CachedProperty && DropZone == EItemDropZone::OntoItem))
 					{
 						const TSharedPtr<const FPropertyBagInstanceDataDetails> DetailsSP = WeakSelf.Pin();
-						const UPropertyBag* ChildBagStruct = DetailsSP ? Private::GetCommonBagStruct(DetailsSP->BagStructProperty) : nullptr;
+						const UPropertyBag* ChildBagStruct = DetailsSP ? GetCommonBagStruct(DetailsSP->BagStructProperty) : nullptr;
 						// Validate these properties are still part of the bag.
 						if (!ChildBagStruct || !ChildBagStruct->FindPropertyDescByProperty(DroppedPropertyDesc.CachedProperty))
 						{
 							return FReply::Unhandled();
 						}
 
-						Private::ApplyChangesToSinglePropertyDesc(
+						ApplyChangesToSinglePropertyDesc(
 							LOCTEXT("DragToChangeCategory", "Change property category"),
 							DroppedPropertyDesc,
 							BagStructProperty,
@@ -886,7 +882,7 @@ void FPropertyBagInstanceDataDetails::OnGroupRowAdded(IDetailGroup& GroupRow, in
 			{
 				using namespace UE::StructUtils;
 
-				Private::ApplyChangesToPropertyDescs(
+				ApplyChangesToPropertyDescs(
 					LOCTEXT("InlineRenameCategory", "Rename category"),
 					BagStructProperty,
 					PropUtils,
@@ -941,7 +937,7 @@ void FPropertyBagInstanceDataDetails::OnGroupRowAdded(IDetailGroup& GroupRow, in
 				.Visibility(bIsFixed ? EVisibility::Collapsed : EVisibility::Visible)
 				.OnClicked_Lambda([GroupName = GroupRow.GetGroupName(), BagStructProperty = BagStructProperty, PropUtils = PropUtils]()
 				{
-					Private::ApplyChangesToPropertyDescs(
+					ApplyChangesToPropertyDescs(
 						LOCTEXT("OnCategoryDeleted", "Delete category"),
 						BagStructProperty,
 						PropUtils,
@@ -1031,7 +1027,7 @@ void FPropertyBagInstanceDataDetails::OnChildRowAdded(IDetailPropertyRow& ChildR
 
 	TWeakPtr<FPropertyBagInstanceDataDetails> WeakSelf = SharedThis<FPropertyBagInstanceDataDetails>(this);
 
-	const UPropertyBag* BagStruct = BagStructProperty ? Private::GetCommonBagStruct(BagStructProperty) : nullptr;
+	const UPropertyBag* BagStruct = BagStructProperty ? GetCommonBagStruct(BagStructProperty) : nullptr;
 	const FProperty* ChildProperty = ChildPropertyHandle->GetProperty();
 	const FPropertyBagPropertyDesc* PropertyDesc = BagStruct ? BagStruct->FindPropertyDescByProperty(ChildProperty) : nullptr;
 
@@ -1225,7 +1221,7 @@ void FPropertyBagInstanceDataDetails::OnChildRowAdded(IDetailPropertyRow& ChildR
 					return;
 				}
 
-				Private::ApplyChangesToPropertyDescs(
+				ApplyChangesToPropertyDescs(
 					FText::Format(LOCTEXT("OnPropertyNameChanged", "Change property name: {0} -> {1}"), FText::FromName(OldName), FText::FromName(NewName)),
 					BagStructProperty,
 					PropUtils,
@@ -1305,7 +1301,7 @@ void FPropertyBagInstanceDataDetails::OnChildRowAdded(IDetailPropertyRow& ChildR
 					.OnClicked_Lambda([ChildPropertyHandle, BagStructProperty = BagStructProperty, PropUtils = PropUtils]()
 					{
 						FProperty* Property = ChildPropertyHandle->GetProperty();
-						Private::ApplyChangesToSinglePropertyDesc(
+						ApplyChangesToSinglePropertyDesc(
 							LOCTEXT("OnPropertyAccessSpecifierChanged", "Set access specifier."),
 							ChildPropertyHandle,
 							BagStructProperty,
@@ -1455,7 +1451,7 @@ void FPropertyBagInstanceDataDetails::OnChildRowAdded(IDetailPropertyRow& ChildR
 						FSlateIcon(FAppStyle::GetAppStyleSetName(), "Icons.Delete"),
 						FUIAction(FExecuteAction::CreateLambda([ChildPropertyHandle, StructProperty = BagStructProperty, PropUtils = PropUtils]()
 						{
-							Private::ApplyChangesToSinglePropertyDesc(
+							ApplyChangesToSinglePropertyDesc(
 								LOCTEXT("DropDownMenuOnCategoryCleared", "Clear property category"),
 								ChildPropertyHandle,
 								StructProperty,
@@ -1511,7 +1507,7 @@ void FPropertyBagInstanceDataDetails::OnChildRowAdded(IDetailPropertyRow& ChildR
 						{
 							if (CommitType == ETextCommit::OnEnter || CommitType == ETextCommit::OnUserMovedFocus)
 							{
-								Private::ApplyChangesToSinglePropertyDesc(
+								ApplyChangesToSinglePropertyDesc(
 									LOCTEXT("DropDownMenuOnCategoryEdited", "Edit property category"),
 									ChildPropertyHandle,
 									StructProperty,
@@ -1561,7 +1557,7 @@ void FPropertyBagInstanceDataDetails::OnChildRowAdded(IDetailPropertyRow& ChildR
 							{
 								if (CommitType != ETextCommit::OnCleared)
 								{
-									Private::ApplyChangesToSinglePropertyDesc(
+									ApplyChangesToSinglePropertyDesc(
 										FText::Format(LOCTEXT("DropDownMenuOnSetSpecifier", "Set metadata specifier '{0}' value"), FText::FromName(SpecifierName)),
 										ChildPropertyHandle,
 										StructProperty,
@@ -1685,7 +1681,7 @@ void FPropertyBagInstanceDataDetails::OnChildRowAdded(IDetailPropertyRow& ChildR
 						using namespace UE::StructUtils;
 
 						const TSharedPtr<FPropertyBagInstanceDataDetails> DetailsSP = WeakSelf.Pin();
-						const UPropertyBag* ChildBagStruct = DetailsSP ? Private::GetCommonBagStruct(DetailsSP->BagStructProperty) : nullptr;
+						const UPropertyBag* ChildBagStruct = DetailsSP ? GetCommonBagStruct(DetailsSP->BagStructProperty) : nullptr;
 						// Validate these properties are still part of the bag.
 						if (!ChildBagStruct
 							|| !ChildBagStruct->FindPropertyDescByProperty(PropertyDesc.CachedProperty)
@@ -1708,7 +1704,7 @@ void FPropertyBagInstanceDataDetails::OnChildRowAdded(IDetailPropertyRow& ChildR
 
 						if (Result == EPropertyBagAlterationResult::Success)
 						{
-							Private::ApplyChangesToSinglePropertyDesc(
+							ApplyChangesToSinglePropertyDesc(
 								LOCTEXT("DragDropReorderProperties", "Reordered properties"),
 								DroppedPropertyDesc,
 								BagStructProperty,
@@ -1958,7 +1954,7 @@ TSharedRef<SWidget> FPropertyBagInstanceDataDetails::OnPropertyNameContent(TShar
 			return;
 		}
 
-		UE::StructUtils::Private::ApplyChangesToPropertyDescs(
+		UE::StructUtils::ApplyChangesToPropertyDescs(
 		LOCTEXT("OnPropertyMoved", "Move Property"), BagStructProperty, PropUtils,
 		[&ChildPropertyHandle, &Delta](TArray<FPropertyBagPropertyDesc>& PropertyDescs)
 		{
@@ -2159,7 +2155,7 @@ TSharedPtr<SWidget> FPropertyBagDetails::MakeAddPropertyWidget(TSharedPtr<IPrope
 					return FReply::Handled();
 				}
 
-				UE::StructUtils::Private::ApplyChangesToPropertyDescs(
+				UE::StructUtils::ApplyChangesToPropertyDescs(
 					LOCTEXT("OnPropertyAdded", "Add Property"), StructProperty, PropUtils,
 					[&NewName, DefaultType](TArray<FPropertyBagPropertyDesc>& PropertyDescs)
 					{
diff --git a/Engine/Source/Editor/StructUtilsEditor/Public/PropertyBagDetails.h b/Engine/Source/Editor/StructUtilsEditor/Public/PropertyBagDetails.h
--- a/Engine/Source/Editor/StructUtilsEditor/Public/PropertyBagDetails.h
+++ b/Engine/Source/Editor/StructUtilsEditor/Public/PropertyBagDetails.h
@@ -74,6 +74,44 @@ namespace UE::StructUtils
 	STRUCTUTILSEDITOR_API FEdGraphPinType GetPropertyDescAsPin(const FPropertyBagPropertyDesc& Desc);
 	/** Creates type selection pill widget. */
 	STRUCTUTILSEDITOR_API TSharedRef<SWidget> CreateTypeSelectionWidget(TSharedPtr<IPropertyHandle> ChildPropertyHandle, const TSharedPtr<IPropertyHandle>& InBagStructProperty, const TSharedPtr<IPropertyUtilities>& InPropUtils, ESelectorType SelectorType, bool bAllowContainers = true);
+
+	/** @return property bag struct common to all edited properties. */
+	STRUCTUTILSEDITOR_API const UPropertyBag* GetCommonBagStruct(TSharedPtr<IPropertyHandle> StructProperty);
+
+	/**
+	 * Allows the mutation of the property descriptors on a Property Bag property via the provided callback.
+	 * Handles the whole transaction internally, including building a new Property Bag Struct (if necessary).
+	 *
+	 * @param SessionName Name of the transaction
+	 * @param StructProperty The Property Bag handle
+	 * @param PropUtils Property Utilities, usually taken from the StructCustomizationUtils in CustomizeHeader
+	 * @param Function The callback that does the mutation
+	 */
+	STRUCTUTILSEDITOR_API void ApplyChangesToPropertyDescs(const FText& SessionName, const TSharedPtr<IPropertyHandle>& StructProperty, const TSharedPtr<IPropertyUtilities>& PropUtils, TFunctionRef<void(TArray<FPropertyBagPropertyDesc>&)> Function);
+
+	/**
+	 * Allows the mutation of a single property descriptor on a Property Bag property via the provided callback.
+	 * Handles the whole transaction internally, including building a new Property Bag Struct (if necessary).
+	 *
+	 * @param SessionName Name of the transaction
+	 * @param PropertyHandle A child handle of the Property Bag to mutate
+	 * @param StructProperty The Property Bag handle
+	 * @param PropUtils Property Utilities, usually taken from the StructCustomizationUtils in CustomizeHeader
+	 * @param Function The callback that does the mutation
+	 */
+	STRUCTUTILSEDITOR_API void ApplyChangesToSinglePropertyDesc(const FText& SessionName, const TSharedPtr<IPropertyHandle> PropertyHandle, const TSharedPtr<IPropertyHandle>& StructProperty, const TSharedPtr<IPropertyUtilities>& PropUtils, TFunctionRef<void(FPropertyBagPropertyDesc&)> Function);
+
+	/**
+	 * Allows the mutation of a single property descriptor on a Property Bag property via the provided callback.
+	 * Handles the whole transaction internally, including building a new Property Bag Struct (if necessary).
+	 *
+	 * @param SessionName Name of the transaction
+	 * @param PropertyDesc The child descriptor to edit
+	 * @param StructProperty The Property Bag handle
+	 * @param PropUtils Property Utilities, usually taken from the StructCustomizationUtils in CustomizeHeader
+	 * @param Function The callback that does the mutation
+	 */
+	STRUCTUTILSEDITOR_API void ApplyChangesToSinglePropertyDesc(const FText& SessionName, const FPropertyBagPropertyDesc& PropertyDesc, const TSharedPtr<IPropertyHandle>& StructProperty, const TSharedPtr<IPropertyUtilities>& PropUtils, TFunctionRef<void(FPropertyBagPropertyDesc&)> Function);
 }
 
 /**
