PropertyBag: Introduce UPropertyBagDetailsExtension through which it's possible to modify the behavior of PropertyBagDetails.
This initial change allows users to specify a PropertyBagInstanceDataDetails to use for the property without having to create a whole new customization for the container.

#rnx

diff --git a/Engine/Source/Editor/StructUtilsEditor/Private/Customizations/PropertyBagDetails.cpp b/Engine/Source/Editor/StructUtilsEditor/Private/Customizations/PropertyBagDetails.cpp
--- a/Engine/Source/Editor/StructUtilsEditor/Private/Customizations/PropertyBagDetails.cpp
+++ b/Engine/Source/Editor/StructUtilsEditor/Private/Customizations/PropertyBagDetails.cpp
@@ -2124,8 +2124,23 @@ void FPropertyBagDetails::CustomizeChildren(TSharedRef<IPropertyHandle> StructPr
 		PRAGMA_ENABLE_DEPRECATION_WARNINGS
 	};
 
+	const UPropertyBagDetailsExtension* Extension = GetDefault<UPropertyBagDetailsExtension>();
+	if (const FString* ExtensionClassName = StructPropertyHandle->GetInstanceMetaData(UE::StructUtils::Metadata::ExtensionClass))
+	{
+		if (UClass* ExtensionClass = FindObject<UClass>(nullptr, *ExtensionClassName))
+		{
+			if (ExtensionClass->IsChildOf(UPropertyBagDetailsExtension::StaticClass()))
+			{
+				if (const UPropertyBagDetailsExtension* CustomExtension = GetDefault<UPropertyBagDetailsExtension>(ExtensionClass))
+				{
+					Extension = CustomExtension;
+				}
+			}
+		}
+	}
+
 	// Show the Value (FInstancedStruct) as child rows.
-	const TSharedRef<FPropertyBagInstanceDataDetails> InstanceDetails = MakeShared<FPropertyBagInstanceDataDetails>(Params);
+	const TSharedRef<FPropertyBagInstanceDataDetails> InstanceDetails = Extension->CreateInstanceDataDetails(Params);
 	StructBuilder.AddCustomBuilder(InstanceDetails);
 }
 
diff --git a/Engine/Source/Editor/StructUtilsEditor/Private/StructUtilsMetadata.cpp b/Engine/Source/Editor/StructUtilsEditor/Private/StructUtilsMetadata.cpp
--- a/Engine/Source/Editor/StructUtilsEditor/Private/StructUtilsMetadata.cpp
+++ b/Engine/Source/Editor/StructUtilsEditor/Private/StructUtilsMetadata.cpp
@@ -21,6 +21,7 @@ namespace UE::StructUtils::Metadata
 	const FLazyName IsPinTypeAcceptedName = "IsPinTypeAccepted";
 	const FLazyName CanRemovePropertyName = "CanRemoveProperty";
 	const FLazyName ChildRowFeaturesName = "ChildRowFeatures";
+	const FLazyName ExtensionClass = "ExtensionClass";
 
 	const FLazyName BaseStructName = "BaseStruct";
 	const FLazyName BaseStructFunctionName = "BaseStructFunction";
diff --git a/Engine/Source/Editor/StructUtilsEditor/Public/PropertyBagDetails.h b/Engine/Source/Editor/StructUtilsEditor/Public/PropertyBagDetails.h
--- a/Engine/Source/Editor/StructUtilsEditor/Public/PropertyBagDetails.h
+++ b/Engine/Source/Editor/StructUtilsEditor/Public/PropertyBagDetails.h
@@ -277,6 +277,18 @@ protected:
 	bool bAllowArrays = true;
 };
 
+UCLASS(MinimalAPI)
+class UPropertyBagDetailsExtension : public UObject
+{
+	GENERATED_BODY()
+
+public:
+	UE_API virtual TSharedRef<FPropertyBagInstanceDataDetails> CreateInstanceDataDetails(const FPropertyBagInstanceDataDetails::FConstructParams& Params) const
+	{
+		return MakeShared<FPropertyBagInstanceDataDetails>(Params);
+	}
+};
+
 /**
  * Specific property bag schema to allow customizing the requirements (e.g. supported containers).
  */
diff --git a/Engine/Source/Editor/StructUtilsEditor/Public/StructUtilsMetadata.h b/Engine/Source/Editor/StructUtilsEditor/Public/StructUtilsMetadata.h
--- a/Engine/Source/Editor/StructUtilsEditor/Public/StructUtilsMetadata.h
+++ b/Engine/Source/Editor/StructUtilsEditor/Public/StructUtilsMetadata.h
@@ -36,6 +36,9 @@ namespace UE::StructUtils::Metadata
 	/// [Bag Property Metadata] ChildRowFeatures: A list of UI features available to the PropertyBag properties when displayed in the details view child rows.
 	extern STRUCTUTILSEDITOR_API const FLazyName ChildRowFeaturesName;
 
+	/// [Bag Property Metadata] ExtensionClass: Path of a class inheriting from UPropertyBagDetailsExtension
+	extern STRUCTUTILSEDITOR_API const FLazyName ExtensionClass;
+
 	/// Common property metadata specifiers
 	/// [Property Metadata] EnableCategoriesName: Required to enable the use of categories in general. Per property.
 	extern STRUCTUTILSEDITOR_API const FLazyName EnableCategoriesName;
