[StructUtils] Add helper functions to generate a UserDefinedStruct from a property bag. Property Bag are great when used internally the same package. It can cause issue when there's a reference outside the package. UserDefinedStruct are better at handling those.

#FailedValidations Preflight

diff --git a/Engine/Source/Editor/StructUtilsEditor/Private/Customizations/PropertyBagDetails.cpp b/Engine/Source/Editor/StructUtilsEditor/Private/Customizations/PropertyBagDetails.cpp
--- a/Engine/Source/Editor/StructUtilsEditor/Private/Customizations/PropertyBagDetails.cpp
+++ b/Engine/Source/Editor/StructUtilsEditor/Private/Customizations/PropertyBagDetails.cpp
@@ -2,7 +2,7 @@
 
 #include "PropertyBagDetails.h"
 
-#include "StructUtilsEditorUtils.h"
+#include "StructUtilsEditorUtilsPrivate.h"
 
 #include "DetailLayoutBuilder.h"
 #include "DetailWidgetRow.h"
@@ -1086,7 +1086,7 @@ void FPropertyBagInstanceDataDetails::OnChildRowAdded(IDetailPropertyRow& ChildR
 			}
 		});
 
-		ChildRow.EditCondition(std::move(EditConditionValue), std::move(OnEditConditionChanged));
+		ChildRow.EditCondition(MoveTemp(EditConditionValue), MoveTemp(OnEditConditionChanged));
 
 		FIsResetToDefaultVisible IsResetVisible = FIsResetToDefaultVisible::CreateLambda([WeakSelf](TSharedPtr<IPropertyHandle> PropertyHandle)
 		{
diff --git a/Engine/Source/Editor/StructUtilsEditor/Private/SInstancedStructPicker.cpp b/Engine/Source/Editor/StructUtilsEditor/Private/SInstancedStructPicker.cpp
--- a/Engine/Source/Editor/StructUtilsEditor/Private/SInstancedStructPicker.cpp
+++ b/Engine/Source/Editor/StructUtilsEditor/Private/SInstancedStructPicker.cpp
@@ -8,7 +8,7 @@
 #include "PropertyCustomizationHelpers.h"
 #include "PropertyHandle.h"
 #include "ScopedTransaction.h"
-#include "StructUtilsEditorUtils.h"
+#include "StructUtilsEditorUtilsPrivate.h"
 #include "StructUtilsMetadata.h"
 #include "Modules/ModuleManager.h"
 #include "StructUtils/InstancedStruct.h"
diff --git a/Engine/Source/Editor/StructUtilsEditor/Private/StructUtilsEditorUtils.cpp b/Engine/Source/Editor/StructUtilsEditor/Private/StructUtilsEditorUtils.cpp
--- a/Engine/Source/Editor/StructUtilsEditor/Private/StructUtilsEditorUtils.cpp
+++ b/Engine/Source/Editor/StructUtilsEditor/Private/StructUtilsEditorUtils.cpp
@@ -2,61 +2,197 @@
 
 #include "StructUtilsEditorUtils.h"
 
-namespace UE::StructUtils
-{
+#include "StructUtilsEditorUtilsPrivate.h"
+#include "StructUtils/PropertyBag.h"
+#include "UObject/Package.h"
+#include "UObject/UnrealType.h"
 
-namespace Private
-{
+#include "Kismet2/BlueprintEditorUtils.h"
+#include "PropertyBagDetails.h"
 
-TOptional<FFindUserFunctionResult> FindUserFunction(const TSharedPtr<IPropertyHandle>& InStructProperty, FName InFuncMetadataName)
+namespace UE::StructUtils::Private
 {
-	FProperty* MetadataProperty = InStructProperty->GetMetaDataProperty();
+	TNotNull<UUserDefinedStruct*> CreateUserDefinedStruct(TNotNull<UObject*> InOuter, FName Name)
+	{
+		UUserDefinedStruct* Struct = NewObject<UUserDefinedStruct>(InOuter, Name, RF_Transactional | RF_Public);
+		check(Struct);
 
-	FFindUserFunctionResult Result{};
+		Struct->EditorData = NewObject<UUserDefinedStructEditorData>(Struct, FName(), RF_Transactional);
+		check(Struct->EditorData);
 
-	if (!MetadataProperty || !MetadataProperty->HasMetaData(InFuncMetadataName))
-	{
-		return {};
+		Struct->Guid = FGuid::NewGuid();
+		Struct->SetMetaData(TEXT("BlueprintType"), TEXT("true"));
+		Struct->Bind();
+		constexpr bool bRelinkExistingProperties = true;
+		Struct->StaticLink(bRelinkExistingProperties);
+
+		Struct->Status = UDSS_Error;
+		return Struct;
 	}
 
-	FString FunctionName = MetadataProperty->GetMetaData(InFuncMetadataName);
-	if (FunctionName.IsEmpty())
+	UUserDefinedStruct* FindUserDefinedStruct(TNotNull<UObject*> InOuter, FName Name)
 	{
-		return {};
+		return Cast<UUserDefinedStruct>(StaticFindObject(UUserDefinedStruct::StaticClass(), InOuter, Name.ToString()));
 	}
 
-	TArray<UObject*> OutObjects;
-	InStructProperty->GetOuterObjects(OutObjects);
+	bool Equal(const FStructVariableDescription& A, const FStructVariableDescription& B)
+	{
+		return A.HasSameType(B)
+			&& A.VarName == B.VarName
+			&& A.VarGuid == B.VarGuid
+			&& A.DefaultValue == B.DefaultValue;
+	}
 
-	// Check for external function references, taken from GetOptions
-	if (FunctionName.Contains(TEXT(".")))
+	FStructVariableDescription CreateUserDefinedVariableFromDesc(const FPropertyBagPropertyDesc& PropertyDesc)
 	{
-		Result.Function = FindObject<UFunction>(nullptr, *FunctionName, EFindObjectFlags::ExactClass);
+		const FEdGraphPinType GraphPin = UE::StructUtils::GetPropertyDescAsPin(PropertyDesc);
+
+		FStructVariableDescription Result;
+		Result.SetPinType(GraphPin);
+		Result.FriendlyName = PropertyDesc.Name.ToString();
+		// FMemberVariableNameHelper adds the ID to the name.
+		Result.VarName = *FString::Printf(TEXT("%s_%s"), *Result.FriendlyName, *PropertyDesc.ID.ToString(EGuidFormats::Digits));
+		Result.VarGuid = PropertyDesc.ID;
 
-		if (ensureMsgf(Result.Function && Result.Function->HasAnyFunctionFlags(EFunctionFlags::FUNC_Static), TEXT("[%s] Didn't find function %s or expected it to be static"), *InFuncMetadataName.ToString(), *FunctionName))
+		return Result;
+	}
+
+	FStructVariableDescription CreateUserDefinedVariableFromDesc(const FInstancedPropertyBag& PropertyBag, const FPropertyBagPropertyDesc& PropertyDesc)
+	{
+		FStructVariableDescription Result = CreateUserDefinedVariableFromDesc(PropertyDesc);
+		if (const FProperty* Property = PropertyBag.GetValue().GetScriptStruct()->FindPropertyByName(PropertyDesc.Name))
 		{
-			UObject* GetOptionsCDO = Result.Function->GetOuterUClass()->GetDefaultObject();
-			Result.Target = GetOptionsCDO;
+			UObject* OwningObject = const_cast<UScriptStruct*>(PropertyBag.GetValue().GetScriptStruct());
+			const bool bValueAsString = FBlueprintEditorUtils::PropertyValueToString(Property, PropertyBag.GetValue().GetMemory(), Result.DefaultValue, OwningObject, PPF_DuplicateVerbatim);
+			ensure(bValueAsString);
 		}
-		else
+		return Result;
+	}
+
+	TArray<FStructVariableDescription> GetPropertyDescs(const FInstancedPropertyBag& PropertyBag, const TArrayView<const FPropertyBagPropertyDesc>  PropertyDescs)
+	{
+		TArray<FStructVariableDescription> VariablesDescriptions;
+		VariablesDescriptions.Reserve(PropertyDescs.Num());
+		TArray<FPropertyBagPropertyDesc> SanitizedPropertyDescs = UPropertyBag::SanitizePropertyDescs(PropertyDescs);
+		for (const FPropertyBagPropertyDesc& Desc : SanitizedPropertyDescs)
 		{
-			return {};
+			VariablesDescriptions.Add(CreateUserDefinedVariableFromDesc(PropertyBag, Desc));
 		}
+		return VariablesDescriptions;
 	}
-	else if (OutObjects.Num() > 0)
+}
+
+namespace UE::StructUtils
+{
+	TValueOrError<UUserDefinedStruct*, void> CreateUserDefinedStructFromDescs(TNotNull<UObject*> PropertyBagPackage, const FInstancedPropertyBag& PropertyBag, FName UserDefinedStructName)
 	{
-		Result.Target = OutObjects[0];
-		Result.Function = Result.Target->GetClass() ? Result.Target->GetClass()->FindFunctionByName(*FunctionName) : nullptr;
+		if (!PropertyBag.IsValid())
+		{
+			return MakeValue(nullptr);
+		}
+
+		if (Private::FindUserDefinedStruct(PropertyBagPackage, UserDefinedStructName))
+		{
+			return MakeError();
+		}
+
+		TArrayView<const FPropertyBagPropertyDesc> PropertyDescs = PropertyBag.GetPropertyBagStruct()->GetPropertyDescs();
+		if (PropertyDescs.IsEmpty())
+		{
+			return MakeValue(nullptr);
+		}
+
+		TNotNull<UUserDefinedStruct*> UserDefinedStruct = Private::CreateUserDefinedStruct(PropertyBagPackage, UserDefinedStructName);
+		TNotNull<UUserDefinedStructEditorData*> EditorData = CastChecked<UUserDefinedStructEditorData>(UserDefinedStruct->EditorData);
+		EditorData->VariablesDescriptions = Private::GetPropertyDescs(PropertyBag, PropertyDescs);
+		FStructureEditorUtils::CompileStructure(UserDefinedStruct);
+
+		return MakeValue(UserDefinedStruct);
 	}
 
-	// Only support native functions
-	if (!ensureMsgf(Result.Function && Result.Function->IsNative(), TEXT("[%s] Didn't find function %s or expected it to be native"), *InFuncMetadataName.ToString(), *FunctionName))
+	bool UpdateUserDefinedStructFromDescs(TNotNull<UUserDefinedStruct*> UserDefinedStruct, const FInstancedPropertyBag& PropertyBag)
 	{
-		return {};
-	}
+		const UUserDefinedStructEditorData* EditorData = CastChecked<UUserDefinedStructEditorData>(UserDefinedStruct->EditorData);
+
+		TArrayView<const FPropertyBagPropertyDesc> PropertyDescs;
+		if (PropertyBag.IsValid())
+		{
+			PropertyDescs = PropertyBag.GetPropertyBagStruct()->GetPropertyDescs();
+		}
 
-	return (Result.Target != nullptr && Result.Function != nullptr) ? Result : TOptional<FFindUserFunctionResult>{};
+		TArray<FStructVariableDescription> VariablesDescriptions = Private::GetPropertyDescs(PropertyBag, PropertyDescs);
+		bool bEqual = EditorData->VariablesDescriptions.Num() == VariablesDescriptions.Num();
+		if (bEqual)
+		{
+			for (int32 Index = 0; Index < VariablesDescriptions.Num(); ++Index)
+			{
+				if (!Private::Equal(VariablesDescriptions[Index], EditorData->VariablesDescriptions[Index]))
+				{
+					bEqual = false;
+					break;
+				}
+			}
+		}
+
+		if (!bEqual)
+		{
+			CastChecked<UUserDefinedStructEditorData>(UserDefinedStruct->EditorData)->VariablesDescriptions = MoveTemp(VariablesDescriptions);
+			FStructureEditorUtils::CompileStructure(UserDefinedStruct);
+		}
+
+		return !bEqual;
+	}
 }
-} // UE::StructUtils::Private
 
-} // UE::StructUtils
+namespace UE::StructUtils::Private
+{
+	TOptional<FFindUserFunctionResult> FindUserFunction(const TSharedPtr<IPropertyHandle>& InStructProperty, FName InFuncMetadataName)
+	{
+		FProperty* MetadataProperty = InStructProperty->GetMetaDataProperty();
+
+		FFindUserFunctionResult Result{};
+
+		if (!MetadataProperty || !MetadataProperty->HasMetaData(InFuncMetadataName))
+		{
+			return {};
+		}
+
+		FString FunctionName = MetadataProperty->GetMetaData(InFuncMetadataName);
+		if (FunctionName.IsEmpty())
+		{
+			return {};
+		}
+
+		TArray<UObject*> OutObjects;
+		InStructProperty->GetOuterObjects(OutObjects);
+
+		// Check for external function references, taken from GetOptions
+		if (FunctionName.Contains(TEXT(".")))
+		{
+			Result.Function = FindObject<UFunction>(nullptr, *FunctionName, EFindObjectFlags::ExactClass);
+
+			if (ensureMsgf(Result.Function && Result.Function->HasAnyFunctionFlags(EFunctionFlags::FUNC_Static), TEXT("[%s] Didn't find function %s or expected it to be static"), *InFuncMetadataName.ToString(), *FunctionName))
+			{
+				UObject* FunctionCDO = Result.Function->GetOuterUClass()->GetDefaultObject();
+				Result.Target = FunctionCDO;
+			}
+			else
+			{
+				return {};
+			}
+		}
+		else if (OutObjects.Num() > 0)
+		{
+			Result.Target = OutObjects[0];
+			Result.Function = Result.Target->GetClass() ? Result.Target->GetClass()->FindFunctionByName(*FunctionName) : nullptr;
+		}
+
+		// Only support native functions
+		if (!ensureMsgf(Result.Function && Result.Function->IsNative(), TEXT("[%s] Didn't find function %s or expected it to be native"), *InFuncMetadataName.ToString(), *FunctionName))
+		{
+			return {};
+		}
+
+		return (Result.Target != nullptr && Result.Function != nullptr) ? Result : TOptional<FFindUserFunctionResult>{};
+	}
+} // UE::StructUtils::Private
diff --git a/Engine/Source/Editor/StructUtilsEditor/Private/StructUtilsEditorUtils.h b/Engine/Source/Editor/StructUtilsEditor/Private/StructUtilsEditorUtils.h
deleted file mode 100644
diff --git a/Engine/Source/Editor/StructUtilsEditor/Private/StructUtilsEditorUtilsPrivate.h b/Engine/Source/Editor/StructUtilsEditor/Private/StructUtilsEditorUtilsPrivate.h
new file mode 100644
--- /dev/null
+++ b/Engine/Source/Editor/StructUtilsEditor/Private/StructUtilsEditorUtilsPrivate.h
@@ -0,0 +1,37 @@
+﻿// Copyright Epic Games, Inc. All Rights Reserved.
+
+#pragma once
+
+#include "PropertyHandle.h"
+#include "Templates/SharedPointer.h"
+
+class UFunction;
+class UObject;
+
+namespace UE::StructUtils::Private
+{
+	
+struct FFindUserFunctionResult
+{
+	UFunction* Function = nullptr;
+	UObject* Target = nullptr;
+};
+
+/** @return the member and the member function that exist on the property handle object. */
+TOptional<FFindUserFunctionResult> FindUserFunction(const TSharedPtr<IPropertyHandle>& InStructProperty, FName InFuncMetadataName);
+
+/** @return true property handle holds struct property of type T.  */
+template<typename T> requires TModels_V<CStaticStructProvider, T>
+bool IsScriptStruct(const TSharedPtr<IPropertyHandle>& PropertyHandle)
+{
+	if (!PropertyHandle)
+	{
+		return false;
+	}
+
+	//@TODO we want IsChild and not IsA
+	const FStructProperty* StructProperty = CastField<const FStructProperty>(PropertyHandle->GetProperty());
+	return StructProperty && StructProperty->Struct->IsA(TBaseStructure<T>::Get()->GetClass());
+}
+
+} // UE::StructUtils::Private
diff --git a/Engine/Source/Editor/StructUtilsEditor/Public/StructUtilsEditorUtils.h b/Engine/Source/Editor/StructUtilsEditor/Public/StructUtilsEditorUtils.h
new file mode 100644
--- /dev/null
+++ b/Engine/Source/Editor/StructUtilsEditor/Public/StructUtilsEditorUtils.h
@@ -0,0 +1,30 @@
+// Copyright Epic Games, Inc. All Rights Reserved.
+
+#pragma once
+
+#include "Templates/ValueOrError.h"
+#include "UserDefinedStructure/UserDefinedStructEditorData.h"
+
+#define UE_API STRUCTUTILSEDITOR_API
+
+class UUserDefinedStruct;
+struct FInstancedPropertyBag;
+struct FPropertyBagPropertyDesc;
+
+namespace UE::StructUtils
+{
+	/**
+	 * @return a new user defined struct created from a property bag.
+	 * Fails if the name is already used by another asset or if the asset could not be created.
+	 * When the property bag is empty, no user defined struct will be created.
+	 */
+	UE_API TValueOrError<UUserDefinedStruct*, void> CreateUserDefinedStructFromDescs(TNotNull<UObject*> PropertyBagOuter, const FInstancedPropertyBag& PropertyBag, FName UserDefinedStructName);
+
+	/**
+	 * Update, if needed, an existing user defined struct to have the exact same properties and default value as the property bag.
+	 * @return true if an update occurs.
+	 */
+	UE_API bool UpdateUserDefinedStructFromDescs(TNotNull<UUserDefinedStruct*> Struct, const FInstancedPropertyBag& PropertyBag);
+}
+
+#undef UE_API
diff --git a/Engine/Source/Editor/UnrealEd/Classes/UserDefinedStructure/UserDefinedStructEditorData.h b/Engine/Source/Editor/UnrealEd/Classes/UserDefinedStructure/UserDefinedStructEditorData.h
--- a/Engine/Source/Editor/UnrealEd/Classes/UserDefinedStructure/UserDefinedStructEditorData.h
+++ b/Engine/Source/Editor/UnrealEd/Classes/UserDefinedStructure/UserDefinedStructEditorData.h
@@ -104,6 +104,9 @@ struct FStructVariableDescription
 
 	UNREALED_API FEdGraphPinType ToPinType() const;
 
+	/** Return true if it generates the same property. */
+	UNREALED_API bool HasSameType(const FStructVariableDescription& Other) const;
+
 	// UE_DEPRECATED(4.17)
 	void PostSerialize(const FArchive& Ar);
 
diff --git a/Engine/Source/Editor/UnrealEd/Private/UserDefinedStructEditorData.cpp b/Engine/Source/Editor/UnrealEd/Private/UserDefinedStructEditorData.cpp
--- a/Engine/Source/Editor/UnrealEd/Private/UserDefinedStructEditorData.cpp
+++ b/Engine/Source/Editor/UnrealEd/Private/UserDefinedStructEditorData.cpp
@@ -78,6 +78,15 @@ FEdGraphPinType FStructVariableDescription::ToPinType() const
 	return FEdGraphPinType(Category, SubCategory, SubCategoryObject.LoadSynchronous(), ContainerType, false, PinValueType);
 }
 
+bool FStructVariableDescription::HasSameType(const FStructVariableDescription& Other) const
+{
+	return Category == Other.Category
+		&& SubCategory == Other.SubCategory
+		&& SubCategoryObject == Other.SubCategoryObject
+		&& PinValueType == Other.PinValueType
+		&& ContainerType == Other.ContainerType;
+}
+
 UUserDefinedStructEditorData::UUserDefinedStructEditorData(const FObjectInitializer& ObjectInitializer)
 	: Super(ObjectInitializer), UniqueNameId(0), CachedStructureChange()
 {
diff --git a/Engine/Source/Runtime/CoreUObject/Private/StructUtils/PropertyBag.cpp b/Engine/Source/Runtime/CoreUObject/Private/StructUtils/PropertyBag.cpp
--- a/Engine/Source/Runtime/CoreUObject/Private/StructUtils/PropertyBag.cpp
+++ b/Engine/Source/Runtime/CoreUObject/Private/StructUtils/PropertyBag.cpp
@@ -3956,35 +3956,13 @@
 	static FTransactionallySafeCriticalSection GPropertyBagLock;
 }
 
-const UPropertyBag* UPropertyBag::GetOrCreateFromDescs(const TConstArrayView<FPropertyBagPropertyDesc> PropertyDescs, const TCHAR* PrefixName)
+TArray<FPropertyBagPropertyDesc> UPropertyBag::SanitizePropertyDescs(const TArrayView<const FPropertyBagPropertyDesc> PropertyDescs)
 {
-	const uint64 BagHash = UE::StructUtils::Private::CalcPropertyDescArrayHash(PropertyDescs);
-	const FString ScriptStructName = PrefixName == nullptr
-		? FString::Printf(TEXT("PropertyBag_%llx"), BagHash)
-		: FString::Printf(TEXT("%s_%llx"), PrefixName, BagHash);
+	TArray<FPropertyBagPropertyDesc> Result = TArray<FPropertyBagPropertyDesc>(PropertyDescs.GetData(), PropertyDescs.Num());
 
-	// We need to linearize this entire function otherwise threads that create bags of identical layouts can view
-	// partially-constructed objects 
-	UE::TScopeLock ScopeLock(UE::StructUtils::Private::GPropertyBagLock);
-
-	// we need to use StaticFindObjectFastInternal with ExclusiveInternalFlags = EInternalObjectFlags::None
-	// here because objects with RF_NeedPostLoad cannot be found by regular FindObject calls as they will have
-	// ExclusiveInternalFlags = EInternalObjectFlags::AsyncLoading when called from game thread.
-	if (UObject* ExistingObject = StaticFindObjectFastInternal(UPropertyBag::StaticClass(), GetTransientPackage(), *ScriptStructName, EFindObjectFlags::ExactClass, RF_NoFlags, EInternalObjectFlags::None))
+	for (FPropertyBagPropertyDesc& Desc : Result)
 	{
-		if (const UPropertyBag* ExistingBag = Cast<UPropertyBag>(ExistingObject))
-		{
-			return ExistingBag;
-		}
-	}
-
-	UPropertyBag* NewBag = NewObject<UPropertyBag>(GetTransientPackage(), *ScriptStructName, RF_Standalone | RF_Transient);
-
-	NewBag->PropertyDescs = PropertyDescs;
-
-	// Fix missing structs, enums, and objects.
-	for (FPropertyBagPropertyDesc& Desc : NewBag->PropertyDescs)
-	{
+		// Fix missing structs, enums, and objects.
 		if (Desc.ValueType == EPropertyBagPropertyType::Struct)
 		{
 			if (Desc.ValueTypeObject == nullptr || Desc.ValueTypeObject->GetClass()->IsChildOf(UScriptStruct::StaticClass()) == false)
@@ -4017,28 +3995,56 @@
 				Desc.ValueTypeObject = UPropertyBagMissingObject::StaticClass();
 			}
 		}
+
+		// Fix ID
+		if (!Desc.ID.IsValid())
+		{
+			Desc.ID = FGuid::NewGuid();
+		}
 	}
-	
+
 	// Remove properties with same name
-	for (int32 Index = 0; Index < NewBag->PropertyDescs.Num() - 1; Index++)
+	for (int32 Index = 0; Index < Result.Num() - 1; Index++)
 	{
-		UE::StructUtils::Private::RemovePropertyByName(NewBag->PropertyDescs, NewBag->PropertyDescs[Index].Name, Index + 1);
+		UE::StructUtils::Private::RemovePropertyByName(Result, Result[Index].Name, Index + 1);
 	}
 
 	// Remove unsupported types
-	NewBag->PropertyDescs.RemoveAll([](const FPropertyBagPropertyDesc& Desc){ return Desc.ValueType == EPropertyBagPropertyType::None; });
-	
+	Result.RemoveAll([](const FPropertyBagPropertyDesc& Desc) { return Desc.ValueType == EPropertyBagPropertyType::None; });
+
+	return Result;
+}
+
+const UPropertyBag* UPropertyBag::GetOrCreateFromDescs(const TConstArrayView<FPropertyBagPropertyDesc> PropertyDescs, const TCHAR* PrefixName)
+{
+	const uint64 BagHash = UE::StructUtils::Private::CalcPropertyDescArrayHash(PropertyDescs);
+	const FString ScriptStructName = PrefixName == nullptr
+		? FString::Printf(TEXT("PropertyBag_%llx"), BagHash)
+		: FString::Printf(TEXT("%s_%llx"), PrefixName, BagHash);
+
+	// We need to linearize this entire function otherwise threads that create bags of identical layouts can view
+	// partially-constructed objects 
+	UE::TScopeLock ScopeLock(UE::StructUtils::Private::GPropertyBagLock);
+
+	// we need to use StaticFindObjectFastInternal with ExclusiveInternalFlags = EInternalObjectFlags::None
+	// here because objects with RF_NeedPostLoad cannot be found by regular FindObject calls as they will have
+	// ExclusiveInternalFlags = EInternalObjectFlags::AsyncLoading when called from game thread.
+	if (UObject* ExistingObject = StaticFindObjectFastInternal(UPropertyBag::StaticClass(), GetTransientPackage(), *ScriptStructName, EFindObjectFlags::ExactClass, RF_NoFlags, EInternalObjectFlags::None))
+	{
+		if (const UPropertyBag* ExistingBag = Cast<UPropertyBag>(ExistingObject))
+		{
+			return ExistingBag;
+		}
+	}
+
+	UPropertyBag* NewBag = NewObject<UPropertyBag>(GetTransientPackage(), *ScriptStructName, RF_Standalone | RF_Transient);
+
+	NewBag->PropertyDescs = SanitizePropertyDescs(PropertyDescs);
+
 	// Add properties (AddCppProperty() adds them backwards in the linked list)
 	for (int32 DescIndex = NewBag->PropertyDescs.Num() - 1; DescIndex >= 0; DescIndex--)
 	{
 		FPropertyBagPropertyDesc& Desc = NewBag->PropertyDescs[DescIndex];
-
-		if (!Desc.ID.IsValid())
-		{
-			Desc.ID = FGuid::NewGuid();
-		}
-
-
 		if (FProperty* NewProperty = UE::StructUtils::Private::CreatePropertyFromDesc(Desc, NewBag))
 		{
 #if WITH_EDITORONLY_DATA
diff --git a/Engine/Source/Runtime/CoreUObject/Public/StructUtils/PropertyBag.h b/Engine/Source/Runtime/CoreUObject/Public/StructUtils/PropertyBag.h
--- a/Engine/Source/Runtime/CoreUObject/Public/StructUtils/PropertyBag.h
+++ b/Engine/Source/Runtime/CoreUObject/Public/StructUtils/PropertyBag.h
@@ -1320,6 +1320,13 @@ public:
 	 */
 	static UE_API const UPropertyBag* GetOrCreateFromDescs(const TConstArrayView<FPropertyBagPropertyDesc> InPropertyDescs, const TCHAR* PrefixName = nullptr);
 
+	/**
+	 * Fix property descriptions that are not properly formed.
+	 * Remove duplicated property descriptions.
+	 * Remove unsupported types.
+	 */
+	static UE_API TArray<FPropertyBagPropertyDesc> SanitizePropertyDescs(const TArrayView<const FPropertyBagPropertyDesc> PropertyDescs);
+
 	/** Returns property descriptions that specify this struct. */
 	TConstArrayView<FPropertyBagPropertyDesc> GetPropertyDescs() const { return PropertyDescs; }
 
