[StateTree] Add button to remove all elements from the tasks, evaluator, condition, transition lists. Add button to remove all properties from the property bag.

diff --git a/Engine/Plugins/Runtime/StateTree/Source/StateTreeEditorModule/Private/Customizations/StateTreeEditorNodeUtils.cpp b/Engine/Plugins/Runtime/StateTree/Source/StateTreeEditorModule/Private/Customizations/StateTreeEditorNodeUtils.cpp
--- a/Engine/Plugins/Runtime/StateTree/Source/StateTreeEditorModule/Private/Customizations/StateTreeEditorNodeUtils.cpp
+++ b/Engine/Plugins/Runtime/StateTree/Source/StateTreeEditorModule/Private/Customizations/StateTreeEditorNodeUtils.cpp
@@ -919,6 +919,19 @@ IDetailCategoryBuilder& MakeArrayCategoryHeader(
 		AddWidget = CreateAddItemButton(AddButtonTooltipText, AddIconColor, ArrayPropertyHandle, DetailBuilder.GetPropertyUtilities());
 	}
 	
+	TSharedRef<SWidget> RemoveAllWidget = PropertyCustomizationHelpers::MakeEmptyButton(
+		FSimpleDelegate::CreateLambda([ArrayPropertyHandle]()
+			{
+				if (ArrayPropertyHandle && ArrayPropertyHandle->IsValidHandle())
+				{
+					if (const TSharedPtr<IPropertyHandleArray> ArrayHandle = ArrayPropertyHandle->AsArray())
+					{
+						ArrayHandle->EmptyArray();
+					}
+				}
+			})
+		);
+
 	const TSharedRef<SHorizontalBox> HeaderContent = SNew(SHorizontalBox);
 
 	if (!IconName.IsNone())
@@ -964,6 +977,15 @@ IDetailCategoryBuilder& MakeArrayCategoryHeader(
 		[
 			AddWidget.ToSharedRef()
 		];
+
+	HeaderContent->AddSlot()
+		.AutoWidth()
+		.Padding(FMargin(0.f, 0.f, 4.f, 0.f))
+		.HAlign(HAlign_Right)
+		.VAlign(VAlign_Center)
+		[
+			RemoveAllWidget
+		];
 	
 	Category.HeaderContent(SNew(SBox)
 		.MinDesiredHeight(30.f)
diff --git a/Engine/Source/Editor/StructUtilsEditor/Private/Customizations/PropertyBagDetails.cpp b/Engine/Source/Editor/StructUtilsEditor/Private/Customizations/PropertyBagDetails.cpp
--- a/Engine/Source/Editor/StructUtilsEditor/Private/Customizations/PropertyBagDetails.cpp
+++ b/Engine/Source/Editor/StructUtilsEditor/Private/Customizations/PropertyBagDetails.cpp
@@ -9,6 +9,7 @@
 #include "IDetailChildrenBuilder.h"
 #include "IDetailGroup.h"
 #include "IPropertyUtilities.h"
+#include "PropertyCustomizationHelpers.h"
 #include "PropertyBagDragDropHandler.h"
 #include "ScopedTransaction.h"
 #include "SDraggableBox.h" // Custom widget for drag and drop sections
@@ -2199,6 +2200,8 @@ void FPropertyBagDetails::CustomizeChildren(TSharedRef<IPropertyHandle> StructPr
 
 TSharedPtr<SWidget> FPropertyBagDetails::MakeAddPropertyWidget(TSharedPtr<IPropertyHandle> InStructProperty, TSharedPtr<IPropertyUtilities> InPropUtils, EPropertyBagPropertyType DefaultType, const FSlateColor IconColor)
 {
+	const FText RemoveAllPropertiesText = LOCTEXT("OnremoveAllProperties", "Remove all Properties");
+
 	return SNew(SHorizontalBox)
 		+SHorizontalBox::Slot()
 		.AutoWidth()
@@ -2239,6 +2242,25 @@ TSharedPtr<SWidget> FPropertyBagDetails::MakeAddPropertyWidget(TSharedPtr<IPrope
 				.Image(FAppStyle::GetBrush("Icons.PlusCircle"))
 				.ColorAndOpacity(IconColor)
 			]
+		]
+		+ SHorizontalBox::Slot()
+		.AutoWidth()
+		.Padding(FMargin(0.f, 0.f, 4.f, 0.f))
+		[
+			PropertyCustomizationHelpers::MakeEmptyButton(
+				FSimpleDelegate::CreateLambda([StructProperty = InStructProperty, PropUtils = InPropUtils, RemoveAllPropertiesText]()
+					{
+						UE::StructUtils::ApplyChangesToPropertyDescs(
+							RemoveAllPropertiesText,
+							StructProperty,
+							PropUtils,
+							[](TArray<FPropertyBagPropertyDesc>& PropertyDescs)
+							{
+								PropertyDescs.Empty();
+							});
+					}),
+				RemoveAllPropertiesText
+				)
 		];
 
 }
