[StructUtils] Add support to specify Property Bag Schema in Metadata

* Useful to limit container creation for example

#rnx

diff --git a/Engine/Source/Editor/StructUtilsEditor/Private/Customizations/PropertyBagDetails.cpp b/Engine/Source/Editor/StructUtilsEditor/Private/Customizations/PropertyBagDetails.cpp
--- a/Engine/Source/Editor/StructUtilsEditor/Private/Customizations/PropertyBagDetails.cpp
+++ b/Engine/Source/Editor/StructUtilsEditor/Private/Customizations/PropertyBagDetails.cpp
@@ -667,8 +667,10 @@ FOptionalSize GetDesiredInputWidgetSizeByType(const TSharedPtr<IPropertyHandle>&
 }
 } // UE::StructUtils::Private
 
-TSharedRef<SWidget> CreateTypeSelectionWidget(TSharedPtr<IPropertyHandle> ChildPropertyHandle, const TSharedPtr<IPropertyHandle>& InBagStructProperty, const TSharedPtr<IPropertyUtilities>& InPropUtils, SPinTypeSelector::ESelectorType SelectorType, const bool bAllowContainers)
+TSharedRef<SWidget> CreateTypeSelectionWidget(TSharedPtr<IPropertyHandle> ChildPropertyHandle, const TSharedPtr<IPropertyHandle>& InBagStructProperty, const TSharedPtr<IPropertyUtilities>& InPropUtils, SPinTypeSelector::ESelectorType SelectorType, const bool bAllowContainers, TSubclassOf<UEdGraphSchema> PropertyBagSchemaClass)
 {
+	const UEdGraphSchema* Schema = PropertyBagSchemaClass ? GetDefault<UEdGraphSchema>(PropertyBagSchemaClass) : GetDefault<UPropertyBagSchema>();
+
 	return SNew(SBox)
 		.HAlign(HAlign_Right)
 		.Padding(FMargin(4, 0))
@@ -686,7 +688,7 @@ TSharedRef<SWidget> CreateTypeSelectionWidget(TSharedPtr<IPropertyHandle> ChildP
 			{
 				return Private::PinInfoChanged(ChildPropertyHandle, BagStructProperty, PropUtils, PinType);
 			})
-			.Schema(GetDefault<UPropertyBagSchema>())
+			.Schema(Schema)
 			.bAllowContainers(bAllowContainers)
 			.SelectorType(SelectorType)
 			.TypeTreeFilter(ETypeTreeFilter::None)
@@ -797,6 +799,7 @@ PRAGMA_DISABLE_DEPRECATION_WARNINGS
 	, bFixedLayout(false)
 	, bAllowArrays(true)
 PRAGMA_ENABLE_DEPRECATION_WARNINGS
+	, PropertyBagSchemaClass(ConstructParams.PropertyBagSchemaClass)
 {
 	ensure(UE::StructUtils::Private::IsScriptStruct<FInstancedPropertyBag>(BagStructProperty));
 	ensure(PropUtils != nullptr);
@@ -1186,7 +1189,7 @@ void FPropertyBagInstanceDataDetails::OnChildRowAdded(IDetailPropertyRow& ChildR
 				.Padding(1, 0)
 				.AutoWidth()
 				[
-					CreateTypeSelectionWidget(ChildPropertyHandle, BagStructProperty, PropUtils, SPinTypeSelector::ESelectorType::Compact, bAllowContainers)
+					CreateTypeSelectionWidget(ChildPropertyHandle, BagStructProperty, PropUtils, SPinTypeSelector::ESelectorType::Compact, bAllowContainers, PropertyBagSchemaClass)
 				];
 		}
 
@@ -1408,7 +1411,8 @@ void FPropertyBagInstanceDataDetails::OnChildRowAdded(IDetailPropertyRow& ChildR
 						BagStructProperty,
 						PropUtils,
 						SPinTypeSelector::ESelectorType::Full,
-						bAllowContainers),
+						bAllowContainers,
+						PropertyBagSchemaClass),
 					/*InLabel=*/FText::GetEmpty());
 				MenuBuilder.EndSection();
 			}
@@ -2035,7 +2039,8 @@ TSharedRef<SWidget> FPropertyBagInstanceDataDetails::OnPropertyNameContent(TShar
 				BagStructProperty,
 				PropUtils,
 				SPinTypeSelector::ESelectorType::Full,
-				bAllowContainers)
+				bAllowContainers,
+				PropertyBagSchemaClass)
 		],
 		FText::GetEmpty());
 
@@ -2101,6 +2106,8 @@ void FPropertyBagDetails::CustomizeHeader(TSharedRef<IPropertyHandle> StructProp
 			? MetaDataProperty->GetBoolMetaData(UE::StructUtils::Metadata::AllowContainersName)
 			: true;
 
+		PropertyBagSchemaClass = FindObject<UClass>(nullptr, MetaDataProperty->GetMetaData(UE::StructUtils::Metadata::PropertyBagSchemaClassName));
+
 		if (MetaDataProperty->HasMetaData(UE::StructUtils::Metadata::DefaultTypeName))
 		{
 			if (UEnum* Enum = StaticEnum<EPropertyBagPropertyType>())
@@ -2174,8 +2181,9 @@ void FPropertyBagDetails::CustomizeChildren(TSharedRef<IPropertyHandle> StructPr
 		.PropUtils = PropUtils,
 		.bAllowContainers = bAllowContainers,
 		PRAGMA_DISABLE_DEPRECATION_WARNINGS
-		.ChildRowFeatures = bFixedLayout ? EPropertyBagChildRowFeatures::Fixed : ChildRowFeatures
+		.ChildRowFeatures = bFixedLayout ? EPropertyBagChildRowFeatures::Fixed : ChildRowFeatures,
 		PRAGMA_ENABLE_DEPRECATION_WARNINGS
+		.PropertyBagSchemaClass = PropertyBagSchemaClass
 	};
 
 	const UPropertyBagDetailsExtension* Extension = GetDefault<UPropertyBagDetailsExtension>();
diff --git a/Engine/Source/Editor/StructUtilsEditor/Private/StructUtilsMetadata.cpp b/Engine/Source/Editor/StructUtilsEditor/Private/StructUtilsMetadata.cpp
--- a/Engine/Source/Editor/StructUtilsEditor/Private/StructUtilsMetadata.cpp
+++ b/Engine/Source/Editor/StructUtilsEditor/Private/StructUtilsMetadata.cpp
@@ -16,6 +16,7 @@ namespace UE::StructUtils::Metadata
 	const FLazyName FixedLayoutName = "FixedLayout";
 	const FLazyName DefaultTypeName = "DefaultType";
 	const FLazyName AllowContainersName = "AllowContainers";
+	const FLazyName PropertyBagSchemaClassName = "PropertyBagSchemaClass";
 	UE_DEPRECATED(5.6, "Use 'AllowContainers' instead.");
 	const FLazyName AllowArraysName = "AllowArrays";
 	const FLazyName IsPinTypeAcceptedName = "IsPinTypeAccepted";
diff --git a/Engine/Source/Editor/StructUtilsEditor/Public/PropertyBagDetails.h b/Engine/Source/Editor/StructUtilsEditor/Public/PropertyBagDetails.h
--- a/Engine/Source/Editor/StructUtilsEditor/Public/PropertyBagDetails.h
+++ b/Engine/Source/Editor/StructUtilsEditor/Public/PropertyBagDetails.h
@@ -73,7 +73,7 @@ namespace UE::StructUtils
 	STRUCTUTILSEDITOR_API void SetPropertyDescFromPin(FPropertyBagPropertyDesc& Desc, const FEdGraphPinType& PinType);
 	STRUCTUTILSEDITOR_API FEdGraphPinType GetPropertyDescAsPin(const FPropertyBagPropertyDesc& Desc);
 	/** Creates type selection pill widget. */
-	STRUCTUTILSEDITOR_API TSharedRef<SWidget> CreateTypeSelectionWidget(TSharedPtr<IPropertyHandle> ChildPropertyHandle, const TSharedPtr<IPropertyHandle>& InBagStructProperty, const TSharedPtr<IPropertyUtilities>& InPropUtils, ESelectorType SelectorType, bool bAllowContainers = true);
+	STRUCTUTILSEDITOR_API TSharedRef<SWidget> CreateTypeSelectionWidget(TSharedPtr<IPropertyHandle> ChildPropertyHandle, const TSharedPtr<IPropertyHandle>& InBagStructProperty, const TSharedPtr<IPropertyUtilities>& InPropUtils, ESelectorType SelectorType, bool bAllowContainers = true, TSubclassOf<UEdGraphSchema> PropertyBagSchemaClass = nullptr);
 
 	/** @return property bag struct common to all edited properties. */
 	STRUCTUTILSEDITOR_API const UPropertyBag* GetCommonBagStruct(TSharedPtr<IPropertyHandle> StructProperty);
@@ -151,6 +151,9 @@ private:
 	UE_DEPRECATED(5.6, "Use 'ChildRowFeatures' instead.")
 	bool bFixedLayout = false;
 	bool bAllowContainers = true;
+	
+	// Default to UPropertyBagSchema, can be set with metadata PropertyBagSchemaClass
+	TSubclassOf<UEdGraphSchema> PropertyBagSchemaClass;
 };
 
 /** 
@@ -170,6 +173,7 @@ public:
 		TSharedPtr<IPropertyUtilities> PropUtils = nullptr;
 		bool bAllowContainers = true;
 		EPropertyBagChildRowFeatures ChildRowFeatures = EPropertyBagChildRowFeatures::Default;
+		TSubclassOf<UEdGraphSchema> PropertyBagSchemaClass;
 	};
 
 	UE_API explicit FPropertyBagInstanceDataDetails(const FConstructParams& ConstructParams);
@@ -273,6 +277,9 @@ protected:
 	bool bFixedLayout = false;
 	UE_DEPRECATED(5.6, "Use 'bAllowContainers' instead.")
 	bool bAllowArrays = true;
+
+	// Default to UPropertyBagSchema, can be set with metadata PropertyBagSchemaClass
+	TSubclassOf<UEdGraphSchema> PropertyBagSchemaClass;
 };
 
 UCLASS(MinimalAPI)
diff --git a/Engine/Source/Editor/StructUtilsEditor/Public/StructUtilsMetadata.h b/Engine/Source/Editor/StructUtilsEditor/Public/StructUtilsMetadata.h
--- a/Engine/Source/Editor/StructUtilsEditor/Public/StructUtilsMetadata.h
+++ b/Engine/Source/Editor/StructUtilsEditor/Public/StructUtilsMetadata.h
@@ -24,6 +24,9 @@ namespace UE::StructUtils::Metadata
 	/// [Bag Property Metadata] AllowContainers: By default it is always allowed to create containers (Array, Set) for properties. Use AllowContainers=false to disable container support.
 	extern STRUCTUTILSEDITOR_API const FLazyName AllowContainersName;
 
+	/// [Bag Property Metadata] PropertyBagSchemaClass: Allow to specify the class for the PropertyBagSchema. Useful to specify which containers are available. If unset or invalid, fallback on UPropertyBagSchema.
+	extern STRUCTUTILSEDITOR_API const FLazyName PropertyBagSchemaClassName;
+
 	UE_DEPRECATED(5.6, "Remains for backwards compatibility. Please see 'AllowContainersName'.")
 	extern STRUCTUTILSEDITOR_API const FLazyName AllowArraysName;
 
