Fix typo in localized text.

#rnx

diff --git a/Engine/Source/Editor/StructUtilsEditor/Private/Customizations/PropertyBagDetails.cpp b/Engine/Source/Editor/StructUtilsEditor/Private/Customizations/PropertyBagDetails.cpp
--- a/Engine/Source/Editor/StructUtilsEditor/Private/Customizations/PropertyBagDetails.cpp
+++ b/Engine/Source/Editor/StructUtilsEditor/Private/Customizations/PropertyBagDetails.cpp
@@ -2208,7 +2208,7 @@ void FPropertyBagDetails::CustomizeChildren(TSharedRef<IPropertyHandle> StructPr
 
 TSharedPtr<SWidget> FPropertyBagDetails::MakeAddPropertyWidget(TSharedPtr<IPropertyHandle> InStructProperty, TSharedPtr<IPropertyUtilities> InPropUtils, EPropertyBagPropertyType DefaultType, const FSlateColor IconColor)
 {
-	const FText RemoveAllPropertiesText = LOCTEXT("OnremoveAllProperties", "Remove all Properties");
+	const FText RemoveAllPropertiesText = LOCTEXT("OnRemoveAllProperties", "Remove all Properties");
 
 	return SNew(SHorizontalBox)
 		+SHorizontalBox::Slot()
