Fixed a bug where changing a UUserDefinedStruct used in an FInstancedStruct property of a UDataTable row, would not reinstance the structs (as they do in non-UDataTable cases), causing data table details panels to crash trying to read new fields in the now-outdated instance memory.

FInstancedStruct reinitialization relies on FStructUtilsEditorModule::PreChange -> UE::StructUtils::Private::VisitReferencedObjects reaching it's FInstancedStruct::AddStructReferencedObjects struct ARO implementation. For UDataTable's, this *should* be handled by UDataTable::AddReferencedObjects, which calls the struct ARO for each row. However, UE::StructUtils::Private::VisitReferencedObjects was only calling the *inner* ARO's for struct properties of UObjects, not the static AddReferencedObjects function for the UObject's themselves e.g: UDataTable::AddReferencedObjects. Now we also call the UObject's ARO during VisitReferencedObjects.

#Synced-CL 48276178

diff --git a/Engine/Source/Editor/StructUtilsEditor/Private/StructUtilsEditorModule.cpp b/Engine/Source/Editor/StructUtilsEditor/Private/StructUtilsEditorModule.cpp
--- a/Engine/Source/Editor/StructUtilsEditor/Private/StructUtilsEditorModule.cpp
+++ b/Engine/Source/Editor/StructUtilsEditor/Private/StructUtilsEditorModule.cpp
@@ -70,6 +70,7 @@ static void VisitReferencedObjects(const UUserDefinedStruct* StructToReinstantia
 		// This sets global variable which read in the AddStructReferencedObjects().
 		FCurrentReinstantiationOuterObjectScope CurrentReinstantiateOuterObjectScope(Object);
 	
+		Object->CallAddReferencedObjects(Collector);
 		Collector.AddPropertyReferencesWithStructARO(Object->GetClass(), Object);
 	}
 
