[StructUtils] Refactor how InstancedStruct, InstancedStructContainer and PropertyBag handle UserDefinedStruct reinstanciation. use only one CDO pass. Reuse the same struct. Serialize the object containing the struct before and desiarilizing it after the compliation. That allows the default values (CDO) to propagate properly to the instance.

#Synced-CL 49772122

diff --git a/Engine/Source/Editor/KismetCompiler/Private/UserDefinedStructureCompilerUtils.cpp b/Engine/Source/Editor/KismetCompiler/Private/UserDefinedStructureCompilerUtils.cpp
--- a/Engine/Source/Editor/KismetCompiler/Private/UserDefinedStructureCompilerUtils.cpp
+++ b/Engine/Source/Editor/KismetCompiler/Private/UserDefinedStructureCompilerUtils.cpp
@@ -15,6 +15,7 @@
 #include "EdMode.h"
 #include "Engine/Blueprint.h"
 #include "Engine/BlueprintGeneratedClass.h"
+#include "StructUtils/PropertyBag.h"
 #include "StructUtils/UserDefinedStruct.h"
 #include "GameFramework/Actor.h"
 #include "HAL/PlatformCrt.h"
@@ -31,6 +32,7 @@
 #include "Logging/LogMacros.h"
 #include "Misc/AssertionMacros.h"
 #include "Misc/Guid.h"
+#include "StructUtils/StructReinstancer.h"
 #include "Templates/Casts.h"
 #include "Templates/SubclassOf.h"
 #include "Templates/UnrealTemplate.h"
@@ -79,97 +81,127 @@ struct FUserDefinedStructureCompilerInner
 		}
 	}
 
-	static void ReplaceStructWithTempDuplicate(
-		UUserDefinedStruct* StructureToReinstance, 
-		TMap<UBlueprint*, FBlueprintUserStructData>& BlueprintsToRecompile,
-		TArray<UUserDefinedStruct*>& ChangedStructs)
+	static TNotNull<UUserDefinedStruct*> CreateDuplicate(TNotNull<UUserDefinedStruct*> StructureToReinstance)
 	{
-		if (StructureToReinstance)
+		UUserDefinedStruct* DuplicatedStruct = nullptr;
 		{
-			UUserDefinedStruct* DuplicatedStruct = NULL;
-			{
-				const FString ReinstancedName = FString::Printf(TEXT("STRUCT_REINST_%s"), *StructureToReinstance->GetName());
-				const FName UniqueName = MakeUniqueObjectName(GetTransientPackage(), UUserDefinedStruct::StaticClass(), FName(*ReinstancedName));
+			const FString ReinstancedName = FString::Printf(TEXT("STRUCT_REINST_%s"), *StructureToReinstance->GetName());
+			const FName UniqueName = MakeUniqueObjectName(GetTransientPackage(), UUserDefinedStruct::StaticClass(), FName(*ReinstancedName));
 
-				TGuardValue<FIsDuplicatingClassForReinstancing, bool> IsDuplicatingClassForReinstancing(GIsDuplicatingClassForReinstancing, true);
-				DuplicatedStruct = (UUserDefinedStruct*)StaticDuplicateObject(StructureToReinstance, GetTransientPackage(), UniqueName, ~RF_Transactional); 
-			}
+			TGuardValue<FIsDuplicatingClassForReinstancing, bool> IsDuplicatingClassForReinstancing(GIsDuplicatingClassForReinstancing, true);
+			DuplicatedStruct = (UUserDefinedStruct*)StaticDuplicateObject(StructureToReinstance, GetTransientPackage(), UniqueName, ~RF_Transactional);
+		}
 
-			DuplicatedStruct->Guid = StructureToReinstance->Guid;
-			DuplicatedStruct->Bind();
-			DuplicatedStruct->StaticLink(true);
-			DuplicatedStruct->PrimaryStruct = StructureToReinstance;
-			DuplicatedStruct->Status = EUserDefinedStructureStatus::UDSS_Duplicate;
-			DuplicatedStruct->SetFlags(RF_Transient);
-			DuplicatedStruct->AddToRoot();
+		DuplicatedStruct->Guid = StructureToReinstance->Guid;
+		DuplicatedStruct->Bind();
+		DuplicatedStruct->StaticLink(true);
+		DuplicatedStruct->PrimaryStruct = StructureToReinstance;
+		DuplicatedStruct->Status = EUserDefinedStructureStatus::UDSS_Duplicate;
+		DuplicatedStruct->SetFlags(RF_Transient);
 
-			CastChecked<UUserDefinedStructEditorData>(DuplicatedStruct->EditorData)->RecreateDefaultInstance();
+		CastChecked<UUserDefinedStructEditorData>(DuplicatedStruct->EditorData)->RecreateDefaultInstance();
 
-			// List of unique classes and structs to regenerate bytecode and property referenced objects list
-			TSet<UStruct*> StructsToRegenerateReferencesFor;
+		return DuplicatedStruct;
+	}
 
-			for (TAllFieldsIterator<FStructProperty> FieldIt(RF_NoFlags, EInternalObjectFlags::Garbage); FieldIt; ++FieldIt)
+	static void ReplaceStructWithTempDuplicate(
+		TNotNull<UScriptStruct*> StructureToReinstance,
+		TNotNull<UScriptStruct*> DuplicatedStruct,
+		TMap<UBlueprint*, FBlueprintUserStructData>& BlueprintsToRecompile,
+		TArray<UScriptStruct*>& ChangedStructs)
+	{
+		// List of unique classes and structs to regenerate bytecode and property referenced objects list
+		TSet<UStruct*> StructsToRegenerateReferencesFor;
+
+		for (TAllFieldsIterator<FStructProperty> FieldIt(RF_NoFlags, EInternalObjectFlags::Garbage); FieldIt; ++FieldIt)
+		{
+			FStructProperty* StructProperty = *FieldIt;
+			if (StructProperty && (StructureToReinstance == StructProperty->Struct))
 			{
-				FStructProperty* StructProperty = *FieldIt;
-				if (StructProperty && (StructureToReinstance == StructProperty->Struct))
+				if (UBlueprintGeneratedClass* OwnerClass = Cast<UBlueprintGeneratedClass>(StructProperty->GetOwnerClass()))
 				{
-					if (UBlueprintGeneratedClass* OwnerClass = Cast<UBlueprintGeneratedClass>(StructProperty->GetOwnerClass()))
+					if (UBlueprint* FoundBlueprint = Cast<UBlueprint>(OwnerClass->ClassGeneratedBy))
 					{
-						if (UBlueprint* FoundBlueprint = Cast<UBlueprint>(OwnerClass->ClassGeneratedBy))
+						ClearStructReferencesInBP(FoundBlueprint, BlueprintsToRecompile);
+						StructProperty->Struct = DuplicatedStruct;
+						StructsToRegenerateReferencesFor.Add(OwnerClass);
+						UStruct* OwnerStruct = StructProperty->GetOwnerStruct();
+						if (OwnerStruct != OwnerClass)
 						{
-							ClearStructReferencesInBP(FoundBlueprint, BlueprintsToRecompile);
-							StructProperty->Struct = DuplicatedStruct;
-							StructsToRegenerateReferencesFor.Add(OwnerClass);
-							UStruct* OwnerStruct = StructProperty->GetOwnerStruct();
-							if (OwnerStruct != OwnerClass)
-							{
-								StructsToRegenerateReferencesFor.Add(OwnerStruct);
-							}
+							StructsToRegenerateReferencesFor.Add(OwnerStruct);
 						}
 					}
-					else if (UUserDefinedStruct* OwnerStruct = Cast<UUserDefinedStruct>(StructProperty->GetOwnerStruct()))
+				}
+				else if (UUserDefinedStruct* OwnerUDS = Cast<UUserDefinedStruct>(StructProperty->GetOwnerStruct()))
+				{
+					check(OwnerUDS != DuplicatedStruct);
+					const bool bValidStruct = (OwnerUDS->GetOutermost() != GetTransientPackage())
+						&& IsValid(OwnerUDS)
+						&& (EUserDefinedStructureStatus::UDSS_Duplicate != OwnerUDS->Status.GetValue());
+
+					if (bValidStruct)
 					{
-						check(OwnerStruct != DuplicatedStruct);
-						const bool bValidStruct = (OwnerStruct->GetOutermost() != GetTransientPackage())
-							&& IsValid(OwnerStruct)
-							&& (EUserDefinedStructureStatus::UDSS_Duplicate != OwnerStruct->Status.GetValue());
+						ChangedStructs.AddUnique(OwnerUDS);
 
-						if (bValidStruct)
+						if (FStructureEditorUtils::FStructEditorManager::ActiveChange != FStructureEditorUtils::EStructureEditorChangeInfo::DefaultValueChanged)
 						{
-							ChangedStructs.AddUnique(OwnerStruct);
+							// Don't change this for a default value only change, it won't get correctly replaced later
+							StructProperty->Struct = DuplicatedStruct;
+							StructsToRegenerateReferencesFor.Add(OwnerUDS);
+						}
+					}
+				}
+				else if (UPropertyBag* OwnerPropertyBag = Cast<UPropertyBag>(StructProperty->GetOwnerStruct()))
+				{
+					static bool Value = true;
+					if (Value == false)
+					{
+						continue;
+					}
+					check(OwnerPropertyBag != DuplicatedStruct);
+
+					ChangedStructs.AddUnique(OwnerPropertyBag);
+					StructProperty->Struct = DuplicatedStruct;
 
-							if (FStructureEditorUtils::FStructEditorManager::ActiveChange != FStructureEditorUtils::EStructureEditorChangeInfo::DefaultValueChanged)
+					const FProperty* StructPropertyOwner = StructProperty->GetTypedOwner<FProperty>();
+					const FProperty* PropertyBagField = StructPropertyOwner != nullptr ? StructPropertyOwner : StructProperty;
+					if (ensure(PropertyBagField != nullptr))
+					{
+						const FPropertyBagPropertyDesc* Desc = OwnerPropertyBag->FindPropertyDescByProperty(PropertyBagField);
+						if (ensure(Desc != nullptr))
+						{
+							if (Desc->ValueTypeObject == StructureToReinstance)
+							{
+								const_cast<FPropertyBagPropertyDesc*>(Desc)->ValueTypeObject = DuplicatedStruct;
+							}
+							if (Desc->KeyTypeObject == StructureToReinstance)
 							{
-								// Don't change this for a default value only change, it won't get correctly replaced later
-								StructProperty->Struct = DuplicatedStruct;
-								StructsToRegenerateReferencesFor.Add(OwnerStruct);
+								const_cast<FPropertyBagPropertyDesc*>(Desc)->KeyTypeObject = DuplicatedStruct;
 							}
 						}
 					}
-					else
-					{
-						UE_LOG(LogK2Compiler, Error, TEXT("ReplaceStructWithTempDuplicate unknown owner"));
-					}
+				}
+				else
+				{
+					UE_LOG(LogK2Compiler, Error, TEXT("ReplaceStructWithTempDuplicate unknown owner"));
 				}
 			}
+		}
 
-			// Make sure we update the list of objects referenced by structs after we replaced the struct in FStructProperties
-			for (UStruct* Struct : StructsToRegenerateReferencesFor)
-			{
-				Struct->CollectBytecodeAndPropertyReferencedObjects();
-			}
-
-			DuplicatedStruct->RemoveFromRoot();
+		// Make sure we update the list of objects referenced by structs after we replaced the struct in FStructProperties
+		for (UStruct* Struct : StructsToRegenerateReferencesFor)
+		{
+			Struct->CollectBytecodeAndPropertyReferencedObjects();
+		}
 
-			for (UBlueprint* Blueprint : TObjectRange<UBlueprint>(RF_ClassDefaultObject, /** bIncludeDerivedClasses */ true, /** InternalExcludeFlags */ EInternalObjectFlags::Garbage))
+		for (UBlueprint* Blueprint : TObjectRange<UBlueprint>(RF_ClassDefaultObject, /** bIncludeDerivedClasses */ true, /** InternalExcludeFlags */ EInternalObjectFlags::Garbage))
+		{
+			if (Blueprint && !BlueprintsToRecompile.Contains(Blueprint))
 			{
-				if (Blueprint && !BlueprintsToRecompile.Contains(Blueprint))
+				FBlueprintEditorUtils::EnsureCachedDependenciesUpToDate(Blueprint);
+				if (Blueprint->CachedUDSDependencies.Contains(StructureToReinstance))
 				{
-					FBlueprintEditorUtils::EnsureCachedDependenciesUpToDate(Blueprint);
-					if (Blueprint->CachedUDSDependencies.Contains(StructureToReinstance))
-					{
-						ClearStructReferencesInBP(Blueprint, BlueprintsToRecompile);
-					}
+					ClearStructReferencesInBP(Blueprint, BlueprintsToRecompile);
 				}
 			}
 		}
@@ -371,38 +403,62 @@ struct FUserDefinedStructureCompilerInner
 		return true;
 	}
 
-	static void BuildDependencyMapAndCompile(const TArray<UUserDefinedStruct*>& ChangedStructs, FCompilerResultsLog& MessageLog)
+	static void BuildDependencyMapAndCompile(const TArray<UScriptStruct*>& ChangedStructs, FCompilerResultsLog& MessageLog)
 	{
 		struct FDependencyMapEntry
 		{
-			UUserDefinedStruct* Struct;
-			TSet<UUserDefinedStruct*> StructuresToWaitFor;
+			UScriptStruct* Struct;
+			TSet<const UScriptStruct*> StructuresToWaitFor;
 
 			FDependencyMapEntry() : Struct(NULL) {}
 
-			void Initialize(UUserDefinedStruct* ChangedStruct, const TArray<UUserDefinedStruct*>& AllChangedStructs) 
+			void Initialize(UScriptStruct* ChangedStruct, const TArray<UScriptStruct*>& AllChangedStructs)
 			{ 
 				Struct = ChangedStruct;
 				check(Struct);
 
-				for (FStructVariableDescription& VarDesc : FStructureEditorUtils::GetVarDesc(Struct))
+				if (UUserDefinedStruct* UserDefinedStruct = Cast<UUserDefinedStruct>(Struct))
+				{
+					for (FStructVariableDescription& VarDesc : FStructureEditorUtils::GetVarDesc(UserDefinedStruct))
+					{
+						UScriptStruct* StructType = Cast<UScriptStruct>(VarDesc.SubCategoryObject.Get());
+						if (StructType && (VarDesc.Category == UEdGraphSchema_K2::PC_Struct) && AllChangedStructs.Contains(StructType))
+						{
+							StructuresToWaitFor.Add(StructType);
+						}
+					}
+				}
+				else if (UPropertyBag* PropertyBag = Cast<UPropertyBag>(Struct))
 				{
-					UUserDefinedStruct* StructType = Cast<UUserDefinedStruct>(VarDesc.SubCategoryObject.Get());
-					if (StructType && (VarDesc.Category == UEdGraphSchema_K2::PC_Struct) && AllChangedStructs.Contains(StructType))
+					for (const FPropertyBagPropertyDesc& Desc : PropertyBag->GetPropertyDescs())
 					{
-						StructuresToWaitFor.Add(StructType);
+						if (Desc.ValueType == EPropertyBagPropertyType::Struct && AllChangedStructs.Contains(Desc.ValueTypeObject))
+						{
+							StructuresToWaitFor.Add(CastChecked<const UScriptStruct>(Desc.ValueTypeObject));
+						}
+						if (Desc.KeyType == EPropertyBagPropertyType::Struct && AllChangedStructs.Contains(Desc.KeyTypeObject))
+						{
+							StructuresToWaitFor.Add(CastChecked<const UScriptStruct>(Desc.KeyTypeObject));
+						}
 					}
 				}
+				else
+				{
+					ensure(false);
+				}
 			}
 		};
 
 		TArray<FDependencyMapEntry> DependencyMap;
-		for (UUserDefinedStruct* ChangedStruct : ChangedStructs)
+		for (UScriptStruct* ChangedStruct : ChangedStructs)
 		{
 			DependencyMap.Add(FDependencyMapEntry());
 			DependencyMap.Last().Initialize(ChangedStruct, ChangedStructs);
 		}
 
+		UE::StructUtils::FStructReinstancer* StructReinstancer = UE::StructUtils::FStructReinstancer::GetInstance();
+		check(StructReinstancer);
+
 		while (DependencyMap.Num())
 		{
 			int32 StructureToCompileIndex = INDEX_NONE;
@@ -415,16 +471,54 @@ struct FUserDefinedStructureCompilerInner
 				}
 			}
 			check(INDEX_NONE != StructureToCompileIndex);
-			UUserDefinedStruct* Struct = DependencyMap[StructureToCompileIndex].Struct;
+			UScriptStruct* Struct = DependencyMap[StructureToCompileIndex].Struct;
 			check(Struct);
 
-			FUserDefinedStructureCompilerInner::CleanAndSanitizeStruct(Struct);
-			FUserDefinedStructureCompilerInner::InnerCompileStruct(Struct, GetDefault<UEdGraphSchema_K2>(), MessageLog);
-			
-			if (UUserDefinedStructEditorData* EditorData = Cast<UUserDefinedStructEditorData>(Struct->EditorData))
+			if (UUserDefinedStruct* UserDefinedStruct = Cast<UUserDefinedStruct>(Struct))
 			{
-				// Ensure that editor data is in sync w/ the current default instance (if valid) so that it can be reinitialized later.
-				EditorData->RefreshValuesFromDefaultInstance();
+				FUserDefinedStructureCompilerInner::CleanAndSanitizeStruct(UserDefinedStruct);
+				FUserDefinedStructureCompilerInner::InnerCompileStruct(UserDefinedStruct, GetDefault<UEdGraphSchema_K2>(), MessageLog);
+				if (UUserDefinedStructEditorData* EditorData = Cast<UUserDefinedStructEditorData>(UserDefinedStruct->EditorData))
+				{
+					// Ensure that editor data is in sync w/ the current default instance (if valid) so that it can be reinitialized later.
+					EditorData->RefreshValuesFromDefaultInstance();
+				}
+
+				StructReinstancer->SetCompiledStruct(Struct, Struct);
+			}
+			else if (UPropertyBag* PropertyBag = Cast<UPropertyBag>(Struct))
+			{
+				//Property bag properties desc use the duplicated UDS or original property bag
+				//replace struct with compiled struct
+				TArrayView<const FPropertyBagPropertyDesc> PreviousPropertyDescs = PropertyBag->GetPropertyDescs();
+				TArray<FPropertyBagPropertyDesc> PropertyDescs = TArray<FPropertyBagPropertyDesc>(PreviousPropertyDescs.GetData(), PreviousPropertyDescs.Num());
+				for (FPropertyBagPropertyDesc& PropDesc : PropertyDescs)
+				{
+					if (PropDesc.ValueType == EPropertyBagPropertyType::Struct)
+					{
+						if (const UScriptStruct* ValueStruct = Cast<const UScriptStruct>(PropDesc.ValueTypeObject))
+						{
+							if (const UScriptStruct* CompiledStruct = StructReinstancer->GetCompiledReinstantingStruct(ValueStruct))
+							{
+								PropDesc.ValueTypeObject = CompiledStruct;
+							}
+						}
+					}
+
+					if (PropDesc.KeyType == EPropertyBagPropertyType::Struct)
+					{
+						if (const UScriptStruct* ValueStruct = Cast<const UScriptStruct>(PropDesc.KeyTypeObject))
+						{
+							if (const UScriptStruct* CompiledStruct = StructReinstancer->GetCompiledReinstantingStruct(ValueStruct))
+							{
+								PropDesc.KeyTypeObject = CompiledStruct;
+							}
+						}
+					}
+				}
+
+				const UPropertyBag* NewPropertyBag = const_cast<UPropertyBag*>(UPropertyBag::GetOrCreateFromDescs(PropertyDescs));
+				StructReinstancer->SetCompiledStruct(Struct, NewPropertyBag);
 			}
 
 			DependencyMap.RemoveAtSwap(StructureToCompileIndex);
@@ -437,28 +531,60 @@ struct FUserDefinedStructureCompilerInner
 	}
 };
 
-void FUserDefinedStructureCompilerUtils::CompileStruct(class UUserDefinedStruct* Struct, class FCompilerResultsLog& MessageLog, bool bForceRecompile)
+void FUserDefinedStructureCompilerUtils::CompileStruct(UUserDefinedStruct* Struct, FCompilerResultsLog& MessageLog, bool bForceRecompile)
 {
 	if (FStructureEditorUtils::UserDefinedStructEnabled() && Struct)
 	{
-		TArray<UUserDefinedStruct*> ChangedStructs; 
+		struct FDuplicatedStruct
+		{
+			TObjectPtr<UScriptStruct> Struct;
+		};
+		TArray<UScriptStruct*> ChangedStructs;
+		TArray<FDuplicatedStruct> DuplicatedStructsToClean;
 		if (FUserDefinedStructureCompilerInner::ShouldBeCompiled(Struct) || bForceRecompile)
 		{
 			ChangedStructs.Add(Struct);
 		}
 
+		UE::StructUtils::FStructReinstancerScope StructReinstancer;
+
 		TMap<UBlueprint*, FUserDefinedStructureCompilerInner::FBlueprintUserStructData> BlueprintsToRecompile;
 		for (int32 StructIdx = 0; StructIdx < ChangedStructs.Num(); ++StructIdx)
 		{
-			UUserDefinedStruct* ChangedStruct = ChangedStructs[StructIdx];
+			UScriptStruct* ChangedStruct = ChangedStructs[StructIdx];
 			if (ChangedStruct)
 			{
-				FStructureEditorUtils::BroadcastPreChange(ChangedStruct);
-				FUserDefinedStructureCompilerInner::ReplaceStructWithTempDuplicate(ChangedStruct, BlueprintsToRecompile, ChangedStructs);
-				ChangedStruct->Status = EUserDefinedStructureStatus::UDSS_Dirty;
+				UUserDefinedStruct* ChangedUDS = Cast<UUserDefinedStruct>(ChangedStruct);
+
+				UScriptStruct* DuplicatedStruct = nullptr;
+				if (ChangedUDS != nullptr)
+				{
+					// Recompiled UDS keep the same instance. Create a duplicate struct to prevent memory corruption.
+					DuplicatedStruct = FUserDefinedStructureCompilerInner::CreateDuplicate(ChangedUDS);
+					DuplicatedStruct->AddToRoot();
+					DuplicatedStructsToClean.Emplace(DuplicatedStruct);
+					FStructureEditorUtils::BroadcastPreChange(ChangedUDS);
+				}
+				else
+				{
+					// Recompiled PropertyBag have a new instance. They can keep their same instance without causing memory coruption.
+					DuplicatedStruct = ChangedStruct;
+				}
+
+				UE::StructUtils::FStructReinstancer::GetInstance()->AddStruct(ChangedStruct, DuplicatedStruct);
+
+				FUserDefinedStructureCompilerInner::ReplaceStructWithTempDuplicate(ChangedStruct, DuplicatedStruct, BlueprintsToRecompile, ChangedStructs);
+				
+				if (ChangedUDS != nullptr)
+				{
+					ChangedUDS->Status = EUserDefinedStructureStatus::UDSS_Dirty;
+				}
 			}
 		}
 
+		// COLLECT ALL UOBJECT THAT REFERENCE THE STRUCTURES
+		UE::StructUtils::FStructReinstancer::GetInstance()->CollectObjects();
+
 		// COMPILE IN PROPER ORDER
 		FUserDefinedStructureCompilerInner::BuildDependencyMapAndCompile(ChangedStructs, MessageLog);
 
@@ -496,14 +622,23 @@ void FUserDefinedStructureCompilerUtils::CompileStruct(class UUserDefinedStruct*
 			FBlueprintEditorUtils::MarkBlueprintAsStructurallyModified(Pair.Key);
 		}
 
-		for (UUserDefinedStruct* ChangedStruct : ChangedStructs)
+		// RESTORE THE UOBJECT THAT REFERENCE THE STRUCTURES
+		UE::StructUtils::FStructReinstancer::GetInstance()->ReinstanceObjects();
+
+		for (UScriptStruct* ChangedStruct : ChangedStructs)
 		{
-			if (ChangedStruct)
+			if (UUserDefinedStruct* ChangedUDS = Cast<UUserDefinedStruct>(ChangedStruct))
 			{
-				FStructureEditorUtils::BroadcastPostChange(ChangedStruct);
-				ChangedStruct->MarkPackageDirty();
+				FStructureEditorUtils::BroadcastPostChange(ChangedUDS);
+				ChangedUDS->MarkPackageDirty();
 			}
 		}
+
+		for (const FDuplicatedStruct& DuplicatedStruct : DuplicatedStructsToClean)
+		{
+			DuplicatedStruct.Struct->RemoveFromRoot();
+			DuplicatedStruct.Struct->StructFlags = EStructFlags(DuplicatedStruct.Struct->StructFlags | STRUCT_NewerVersionExists);
+		}
 	}
 }
 
@@ -514,25 +649,9 @@ void FUserDefinedStructureCompilerUtils::ReplaceStructWithTempDuplicateByPredica
 {
 	if (StructureToReinstance)
 	{
-		UUserDefinedStruct* DuplicatedStruct = NULL;
-		{
-			const FString ReinstancedName = FString::Printf(TEXT("STRUCT_REINST_%s"), *StructureToReinstance->GetName());
-			const FName UniqueName = MakeUniqueObjectName(GetTransientPackage(), UUserDefinedStruct::StaticClass(), FName(*ReinstancedName));
-
-			TGuardValue<FIsDuplicatingClassForReinstancing, bool> IsDuplicatingClassForReinstancing(GIsDuplicatingClassForReinstancing, true);
-			DuplicatedStruct = (UUserDefinedStruct*)StaticDuplicateObject(StructureToReinstance, GetTransientPackage(), UniqueName, ~RF_Transactional); 
-		}
-
-		DuplicatedStruct->Guid = StructureToReinstance->Guid;
-		DuplicatedStruct->Bind();
-		DuplicatedStruct->StaticLink(true);
-		DuplicatedStruct->PrimaryStruct = StructureToReinstance;
-		DuplicatedStruct->Status = EUserDefinedStructureStatus::UDSS_Duplicate;
-		DuplicatedStruct->SetFlags(RF_Transient);
+		UUserDefinedStruct* DuplicatedStruct = FUserDefinedStructureCompilerInner::CreateDuplicate(StructureToReinstance);
 		DuplicatedStruct->AddToRoot();
 
-		CastChecked<UUserDefinedStructEditorData>(DuplicatedStruct->EditorData)->RecreateDefaultInstance();
-
 		// List of unique classes and structs to regenerate
 		TSet<UStruct*> StructsToRegenerateReferencesFor;
 
diff --git a/Engine/Source/Editor/StructUtilsEditor/Private/StructUtilsEditorModule.cpp b/Engine/Source/Editor/StructUtilsEditor/Private/StructUtilsEditorModule.cpp
--- a/Engine/Source/Editor/StructUtilsEditor/Private/StructUtilsEditorModule.cpp
+++ b/Engine/Source/Editor/StructUtilsEditor/Private/StructUtilsEditorModule.cpp
@@ -40,96 +40,8 @@ void FStructUtilsEditorModule::ShutdownModule()
 	}
 }
 
-namespace UE::StructUtils::Private
-{
-
-static void VisitReferencedObjects(const UUserDefinedStruct* StructToReinstantiate)
-{
-	// Helper preference collector, does not collect anything, but makes sure AddStructReferencedObjects() gets called e.g. on instanced struct. 
-	class FVisitorReferenceCollector : public FReferenceCollector
-	{
-	public:
-		virtual bool IsIgnoringArchetypeRef() const override { return false; }
-		virtual bool IsIgnoringTransient() const override { return false; }
-
-		virtual void HandleObjectReference(UObject*& Object, const UObject* ReferencingObject, const FProperty* ReferencingProperty) override
-		{
-			// Empty
-		}
-	};
-
-	FVisitorReferenceCollector Collector;
-
-	// This sets global variable which read in the AddStructReferencedObjects().
-	FStructureToReinstantiateScope StructureToReinstantiateScope(StructToReinstantiate);
-
-	for (TObjectIterator<UObject> ObjectIt; ObjectIt; ++ObjectIt)
-	{
-		UObject* Object = *ObjectIt;
-	
-		// This sets global variable which read in the AddStructReferencedObjects().
-		FCurrentReinstantiationOuterObjectScope CurrentReinstantiateOuterObjectScope(Object);
-	
-		Object->CallAddReferencedObjects(Collector);
-		Collector.AddPropertyReferencesWithStructARO(Object->GetClass(), Object);
-	}
-
-	// Handle CDOs and sparse class data 
-	for (TObjectIterator<UClass> ClassIt; ClassIt; ++ClassIt)
-	{
-		UClass* Class = *ClassIt;
-
-		// Handle Sparse class data
-		if (void* SparseData = const_cast<void*>(Class->GetSparseClassData(EGetSparseClassDataMethod::ReturnIfNull)))
-		{
-			FCurrentReinstantiationOuterObjectScope CurrentReinstantiateOuterObjectScope(Class);
-			const UScriptStruct* SparseDataStruct = Class->GetSparseClassDataStruct();
-			Collector.AddPropertyReferencesWithStructARO(SparseDataStruct, SparseData);
-		}
-
-		// Handle CDO
-		if (UObject* CDO = Class->GetDefaultObject())
-		{
-			FCurrentReinstantiationOuterObjectScope CurrentReinstantiateOuterObjectScope(CDO);
-			Collector.AddPropertyReferencesWithStructARO(Class, CDO);
-		}
-	}
-}
-
-} // UE::StructUtils::Private
-
 void FStructUtilsEditorModule::PreChange(const UUserDefinedStruct* StructToReinstantiate, FStructureEditorUtils::EStructureEditorChangeInfo Info)
 {
-	if (!StructToReinstantiate)
-	{
-		return;
-	}
-
-	// Make a duplicate of the existing struct, and point all instances of the struct to point to the duplicate.
-	// This is done because the original struct will be changed.
-	UUserDefinedStruct* DuplicatedStruct = nullptr;
-	{
-		const FString ReinstantiatedName = FString::Printf(TEXT("STRUCT_REINST_%s"), *StructToReinstantiate->GetName());
-		const FName UniqueName = MakeUniqueObjectName(GetTransientPackage(), UUserDefinedStruct::StaticClass(), FName(*ReinstantiatedName));
-
-		TGuardValue<FIsDuplicatingClassForReinstancing, bool> IsDuplicatingClassForReinstancing(GIsDuplicatingClassForReinstancing, true);
-		DuplicatedStruct = (UUserDefinedStruct*)StaticDuplicateObject(StructToReinstantiate, GetTransientPackage(), UniqueName, ~RF_Transactional); 
-
-		DuplicatedStruct->Guid = StructToReinstantiate->Guid;
-		DuplicatedStruct->Bind();
-		DuplicatedStruct->StaticLink(true);
-		DuplicatedStruct->PrimaryStruct = const_cast<UUserDefinedStruct*>(StructToReinstantiate);
-		DuplicatedStruct->Status = EUserDefinedStructureStatus::UDSS_Duplicate;
-		DuplicatedStruct->SetFlags(RF_Transient);
-		DuplicatedStruct->AddToRoot();
-	}
-
-	UUserDefinedStructEditorData* DuplicatedEditorData = CastChecked<UUserDefinedStructEditorData>(DuplicatedStruct->EditorData);
-	DuplicatedEditorData->RecreateDefaultInstance();
-
-	UE::StructUtils::Private::VisitReferencedObjects(DuplicatedStruct);
-	
-	DuplicatedStruct->RemoveFromRoot();
 }
 
 void FStructUtilsEditorModule::PostChange(const UUserDefinedStruct* StructToReinstantiate, FStructureEditorUtils::EStructureEditorChangeInfo Info)
@@ -139,8 +51,6 @@ void FStructUtilsEditorModule::PostChange(const UUserDefinedStruct* StructToRein
 		return;
 	}
 
-	UE::StructUtils::Private::VisitReferencedObjects(StructToReinstantiate);
-
 	if (UE::StructUtils::Delegates::OnUserDefinedStructReinstanced.IsBound())
 	{
 		UE::StructUtils::Delegates::OnUserDefinedStructReinstanced.Broadcast(*StructToReinstantiate);
diff --git a/Engine/Source/Editor/StructUtilsEditor/Private/StructUtilsEditorUtils.cpp b/Engine/Source/Editor/StructUtilsEditor/Private/StructUtilsEditorUtils.cpp
--- a/Engine/Source/Editor/StructUtilsEditor/Private/StructUtilsEditorUtils.cpp
+++ b/Engine/Source/Editor/StructUtilsEditor/Private/StructUtilsEditorUtils.cpp
@@ -84,14 +84,14 @@ namespace UE::StructUtils::Private
 
 namespace UE::StructUtils
 {
-	TValueOrError<UUserDefinedStruct*, void> CreateUserDefinedStructFromDescs(TNotNull<UObject*> PropertyBagPackage, const FInstancedPropertyBag& PropertyBag, FName UserDefinedStructName)
+	TValueOrError<UUserDefinedStruct*, void> CreateUserDefinedStructFromDescs(TNotNull<UObject*> UDSPackage, const FInstancedPropertyBag& PropertyBag, FName UserDefinedStructName)
 	{
 		if (!PropertyBag.IsValid())
 		{
 			return MakeValue(nullptr);
 		}
 
-		if (Private::FindUserDefinedStruct(PropertyBagPackage, UserDefinedStructName))
+		if (Private::FindUserDefinedStruct(UDSPackage, UserDefinedStructName))
 		{
 			return MakeError();
 		}
@@ -102,7 +102,7 @@ namespace UE::StructUtils
 			return MakeValue(nullptr);
 		}
 
-		TNotNull<UUserDefinedStruct*> UserDefinedStruct = Private::CreateUserDefinedStruct(PropertyBagPackage, UserDefinedStructName);
+		TNotNull<UUserDefinedStruct*> UserDefinedStruct = Private::CreateUserDefinedStruct(UDSPackage, UserDefinedStructName);
 		TNotNull<UUserDefinedStructEditorData*> EditorData = CastChecked<UUserDefinedStructEditorData>(UserDefinedStruct->EditorData);
 		EditorData->VariablesDescriptions = Private::GetPropertyDescs(PropertyBag, PropertyDescs);
 		FStructureEditorUtils::CompileStructure(UserDefinedStruct);
@@ -112,32 +112,46 @@ namespace UE::StructUtils
 
 	bool UpdateUserDefinedStructFromDescs(TNotNull<UUserDefinedStruct*> UserDefinedStruct, const FInstancedPropertyBag& PropertyBag)
 	{
-		const UUserDefinedStructEditorData* EditorData = CastChecked<UUserDefinedStructEditorData>(UserDefinedStruct->EditorData);
-
-		TArrayView<const FPropertyBagPropertyDesc> PropertyDescs;
-		if (PropertyBag.IsValid())
+		bool bEqual = true;
+		UUserDefinedStructEditorData* EditorData = CastChecked<UUserDefinedStructEditorData>(UserDefinedStruct->EditorData);
+		if (ensure(EditorData))
 		{
-			PropertyDescs = PropertyBag.GetPropertyBagStruct()->GetPropertyDescs();
-		}
+			TArrayView<const FPropertyBagPropertyDesc> PropertyDescs;
+			if (PropertyBag.IsValid())
+			{
+				PropertyDescs = PropertyBag.GetPropertyBagStruct()->GetPropertyDescs();
+			}
 
-		TArray<FStructVariableDescription> VariablesDescriptions = Private::GetPropertyDescs(PropertyBag, PropertyDescs);
-		bool bEqual = EditorData->VariablesDescriptions.Num() == VariablesDescriptions.Num();
-		if (bEqual)
-		{
-			for (int32 Index = 0; Index < VariablesDescriptions.Num(); ++Index)
+			TArray<FStructVariableDescription> VariablesDescriptions = Private::GetPropertyDescs(PropertyBag, PropertyDescs);
+			bEqual = EditorData->VariablesDescriptions.Num() == VariablesDescriptions.Num();
+			if (bEqual)
 			{
-				if (!Private::Equal(VariablesDescriptions[Index], EditorData->VariablesDescriptions[Index]))
+				for (int32 Index = 0; Index < VariablesDescriptions.Num(); ++Index)
 				{
-					bEqual = false;
-					break;
+					if (!Private::Equal(VariablesDescriptions[Index], EditorData->VariablesDescriptions[Index]))
+					{
+						bEqual = false;
+						break;
+					}
 				}
 			}
-		}
 
-		if (!bEqual)
-		{
-			CastChecked<UUserDefinedStructEditorData>(UserDefinedStruct->EditorData)->VariablesDescriptions = MoveTemp(VariablesDescriptions);
-			FStructureEditorUtils::CompileStructure(UserDefinedStruct);
+			if (!bEqual)
+			{
+				// Update the Current Value before setting the Default value.
+				if (const uint8* StructData = UserDefinedStruct->GetDefaultInstance())
+				{
+					for (FStructVariableDescription& Description : VariablesDescriptions)
+					{
+						if (FProperty* Property = EditorData->FindProperty(UserDefinedStruct, Description.VarName))
+						{
+							FBlueprintEditorUtils::PropertyValueToString(Property, StructData, Description.CurrentDefaultValue, UserDefinedStruct);
+						}
+					}
+				}
+				EditorData->VariablesDescriptions = MoveTemp(VariablesDescriptions);
+				FStructureEditorUtils::CompileStructure(UserDefinedStruct);
+			}
 		}
 
 		return !bEqual;
diff --git a/Engine/Source/Editor/StructUtilsEditor/Private/Tests/StructReinstancerTests.cpp b/Engine/Source/Editor/StructUtilsEditor/Private/Tests/StructReinstancerTests.cpp
new file mode 100644
--- /dev/null
+++ b/Engine/Source/Editor/StructUtilsEditor/Private/Tests/StructReinstancerTests.cpp
@@ -0,0 +1,506 @@
+// Copyright Epic Games, Inc. All Rights Reserved.
+
+#include "Tests/StructReinstancerTests.h"
+
+#include "Misc/AutomationTest.h"
+
+#include "StructUtils/StructReinstancer.h"
+#include "StructUtilsEditorUtils.h"
+
+#define LOCTEXT_NAMESPACE "StructReinstancerTests"
+
+namespace UE::StructUtils::Private
+{
+	template<typename T>
+	T GetPropertyValue_InContainer(TNotNull<const FProperty*> Property, TNotNull<const uint8*> Container)
+	{
+		T LocalValue{};
+		Property->GetValue_InContainer(Container, &LocalValue);
+		return LocalValue;
+	}
+
+	template<typename T>
+	TOptional<T> TryGetUDSPropertyValue_InContainer(TNotNull<const UUserDefinedStruct*> UserDefinedStruct, FName PropertyName, const uint8* Container)
+	{
+		if (Container == nullptr)
+		{
+			return {};
+		}
+
+		const UUserDefinedStructEditorData* EditorData = Cast<UUserDefinedStructEditorData>(UserDefinedStruct->EditorData);
+		if (EditorData == nullptr)
+		{
+			return {};
+		}
+
+		const FProperty* Property = EditorData->FindProperty(UserDefinedStruct, PropertyName);
+		if (Property == nullptr)
+		{
+			return {};
+		}
+
+		T LocalValue{};
+		Property->GetValue_InContainer(Container, &LocalValue);
+		return LocalValue;
+	}
+
+	template<typename T>
+	TOptional<T> TryGetPropertyValue_InContainer(TNotNull<const UScriptStruct*> ScriptStruct, FName PropertyName, const uint8* Container)
+	{
+		if (Container == nullptr)
+		{
+			return {};
+		}
+
+		const FProperty* Property = ScriptStruct->FindPropertyByName(PropertyName);
+		if (Property == nullptr)
+		{
+			return {};
+		}
+
+		T LocalValue{};
+		Property->GetValue_InContainer(Container, &LocalValue);
+		return LocalValue;
+	}
+
+	template<typename T>
+	bool TrySetUDSPropertyValue_InContainer(TNotNull<UUserDefinedStruct*> UserDefinedStruct, FName PropertyName, uint8* Container, T Value)
+	{
+		if (Container == nullptr)
+		{
+			return false;
+		}
+
+		const UUserDefinedStructEditorData* EditorData = Cast<UUserDefinedStructEditorData>(UserDefinedStruct->EditorData);
+		if (EditorData == nullptr)
+		{
+			return false;
+		}
+
+		const FProperty* Property = EditorData->FindProperty(UserDefinedStruct, PropertyName);
+		if (Property == nullptr)
+		{
+			return false;
+		}
+
+		Property->SetValue_InContainer(Container, &Value);
+		return true;
+	}
+
+	IMPLEMENT_SIMPLE_AUTOMATION_TEST(FTest_UserDefinedStructReinstancer
+		, "System.StructUtils.Reinstancer.UserDefinedStruct"
+		, EAutomationTestFlags::EditorContext | EAutomationTestFlags::EngineFilter);
+	bool FTest_UserDefinedStructReinstancer::RunTest(const FString& Parameters)
+	{
+		static const FName AName("A");
+		static const FName BName("B");
+		static const FName CName("C");
+
+		TStrongObjectPtr<UTestObjectThatContainsStruct> UDSContainer(NewObject<UTestObjectThatContainsStruct>());
+
+		// Create UDS
+		FInstancedPropertyBag BagA;
+		BagA.AddProperties({
+			{ AName, EPropertyBagPropertyType::String }, // add a string to check if the stomp allocator might find something
+			{ BName, EPropertyBagPropertyType::Int32 }
+			});
+		BagA.SetValueString(AName, TEXT("A"));
+		BagA.SetValueInt32(BName, 2);
+
+		UUserDefinedStruct* UserDefinedStruct = nullptr;
+		{
+			TValueOrError<UUserDefinedStruct*, void> Result = UE::StructUtils::CreateUserDefinedStructFromDescs(UDSContainer.Get(), BagA, "UDS");
+
+			UTEST_TRUE("UUserDefinedStruct created", Result.HasValue());
+			UserDefinedStruct = Result.GetValue();
+			UTEST_TRUE("UUserDefinedStruct valid", UserDefinedStruct != nullptr);
+		}
+
+		// Assign to test object
+		TStrongObjectPtr<UTestObjectThatContainsStruct> TestInstancedObject1(NewObject<UTestObjectThatContainsStruct>());
+		TestInstancedObject1->InstancedStruct = UserDefinedStruct;
+		TStrongObjectPtr<UTestObjectThatContainsStruct> TestInstancedObject2(NewObject<UTestObjectThatContainsStruct>());
+		TestInstancedObject2->Container.InstancedStruct = UserDefinedStruct;
+		TStrongObjectPtr<UTestObjectThatContainsStruct> TestContainerObject1(NewObject<UTestObjectThatContainsStruct>());
+		TestContainerObject1->InstancedStructContainer.InsertAt(0, TArrayView<FInstancedStruct>(&TestInstancedObject1->InstancedStruct, 1));
+		TStrongObjectPtr<UTestObjectThatContainsStruct> TestContainerObject2(NewObject<UTestObjectThatContainsStruct>());
+		TestContainerObject2->Container.InstancedStructContainer.InsertAt(0, TArrayView<FInstancedStruct>(&TestInstancedObject1->InstancedStruct, 1));
+		TStrongObjectPtr<UTestObjectThatContainsStruct> TestBagObject1(NewObject<UTestObjectThatContainsStruct>());
+		TestBagObject1->InstancedPropertyBag.AddProperties({{"MyStruct1", EPropertyBagPropertyType::Struct, UserDefinedStruct}});
+		TStrongObjectPtr<UTestObjectThatContainsStruct> TestBagObject2(NewObject<UTestObjectThatContainsStruct>());
+		TestBagObject2->Container.InstancedPropertyBag.AddProperties({ {"MyStruct2", EPropertyBagPropertyType::Struct, UserDefinedStruct} });
+
+		{
+			UTEST_TRUE("TestInstancedObject1.InstancedStruct struct", TestInstancedObject1->InstancedStruct.GetScriptStruct() == UserDefinedStruct);
+			UTEST_TRUE("TestInstancedObject1.Container.InstancedStruct struct", TestInstancedObject1->Container.InstancedStruct.GetScriptStruct() == nullptr);
+			UTEST_TRUE("TestInstancedObject2.InstancedStruct struct", TestInstancedObject2->Container.InstancedStruct.GetScriptStruct() == UserDefinedStruct);
+			UTEST_TRUE("TestInstancedObject2.Container.InstancedStruct struct", TestInstancedObject2->InstancedStruct.GetScriptStruct() == nullptr);
+
+			UTEST_TRUE("TestContainerObject1.InstancedStructContainer struct", TestContainerObject1->InstancedStructContainer[0].GetScriptStruct() == UserDefinedStruct);
+			UTEST_TRUE("TestContainerObject1.Container.InstancedStructContainer struct", TestContainerObject1->Container.InstancedStructContainer.Num() == 0);
+			UTEST_TRUE("TestContainerObject2.InstancedStructContainer struct", TestContainerObject2->Container.InstancedStructContainer[0].GetScriptStruct() == UserDefinedStruct);
+			UTEST_TRUE("TestContainerObject2.Container.InstancedStructContainer struct", TestContainerObject2->InstancedStructContainer.Num() == 0);
+
+			UTEST_TRUE("TestBagObject1.InstancedPropertyBag struct", TestBagObject1->InstancedPropertyBag.GetValue().GetScriptStruct() != UserDefinedStruct);
+			UTEST_TRUE("TestBagObject1.Container.InstancedPropertyBag struct", TestBagObject1->Container.InstancedPropertyBag.GetValue().GetScriptStruct() == nullptr);
+			UTEST_TRUE("TestBagObject2.InstancedPropertyBag struct", TestBagObject2->Container.InstancedPropertyBag.GetValue().GetScriptStruct() != UserDefinedStruct);
+			UTEST_TRUE("TestBagObject2.Container.InstancedPropertyBag struct", TestBagObject2->InstancedPropertyBag.GetValue().GetScriptStruct() == nullptr);
+
+			UTEST_TRUE("TestInstancedObject1 B Value", 2 == TryGetUDSPropertyValue_InContainer<int32>(UserDefinedStruct, BName, TestInstancedObject1->InstancedStruct.GetMemory()).Get(0));
+			UTEST_TRUE("TestInstancedObject1 B Value", !TryGetUDSPropertyValue_InContainer<int32>(UserDefinedStruct, BName, TestInstancedObject1->Container.InstancedStruct.GetMemory()).IsSet());
+			UTEST_TRUE("TestInstancedObject2 B Value", !TryGetUDSPropertyValue_InContainer<int32>(UserDefinedStruct, BName, TestInstancedObject2->InstancedStruct.GetMemory()).IsSet());
+			UTEST_TRUE("TestInstancedObject2 B Value", 2 == TryGetUDSPropertyValue_InContainer<int32>(UserDefinedStruct, BName, TestInstancedObject2->Container.InstancedStruct.GetMemory()).Get(0));
+
+			UTEST_TRUE("TestContainerObject1 B Value", 2 == TryGetUDSPropertyValue_InContainer<int32>(UserDefinedStruct, BName, TestContainerObject1->InstancedStructContainer[0].GetMemory()).Get(0));
+			UTEST_TRUE("TestContainerObject2 B Value", 2 == TryGetUDSPropertyValue_InContainer<int32>(UserDefinedStruct, BName, TestContainerObject2->Container.InstancedStructContainer[0].GetMemory()).Get(0));
+
+			TValueOrError<FStructView, EPropertyBagResult> BagResult1 = TestBagObject1->InstancedPropertyBag.GetValueStruct("MyStruct1");
+			UTEST_TRUE(TEXT("BagResult1"), BagResult1.HasValue());
+			UTEST_TRUE(TEXT("BagResult1 struct"), BagResult1.GetValue().GetScriptStruct() == UserDefinedStruct);
+			UTEST_TRUE(TEXT("BagResult1 B Value"), 2 == TryGetUDSPropertyValue_InContainer<int32>(UserDefinedStruct, BName, BagResult1.GetValue().GetMemory()).Get(0));
+			TValueOrError<FStructView, EPropertyBagResult> BagResult2 = TestBagObject2->Container.InstancedPropertyBag.GetValueStruct("MyStruct2");
+			UTEST_TRUE(TEXT("BagResult2"), BagResult2.HasValue());
+			UTEST_TRUE(TEXT("BagResult2 struct"), BagResult2.GetValue().GetScriptStruct() == UserDefinedStruct);
+			UTEST_TRUE(TEXT("BagResult2 B Value"), 2 == TryGetUDSPropertyValue_InContainer<int32>(UserDefinedStruct, BName, BagResult2.GetValue().GetMemory()).Get(0));
+		}
+
+		// Add new property
+		{
+			BagA.AddProperties({ { CName, EPropertyBagPropertyType::Int32 } });
+			BagA.SetValueInt32(BName, 22);
+			BagA.SetValueInt32(CName, 3);
+
+			const bool bModified = UE::StructUtils::UpdateUserDefinedStructFromDescs(UserDefinedStruct, BagA);
+			UTEST_TRUE("UUserDefinedStruct modified", bModified);
+
+			{
+				UTEST_TRUE(TEXT("TestInstancedObject1.InstancedStruct struct"), TestInstancedObject1->InstancedStruct.GetScriptStruct() == UserDefinedStruct);
+				UTEST_TRUE(TEXT("TestInstancedObject1.Container.InstancedStruct struct"), TestInstancedObject1->Container.InstancedStruct.GetScriptStruct() == nullptr);
+				UTEST_TRUE(TEXT("TestInstancedObject2.InstancedStruct struct"), TestInstancedObject2->Container.InstancedStruct.GetScriptStruct() == UserDefinedStruct);
+				UTEST_TRUE(TEXT("TestInstancedObject2.Container.InstancedStruct struct"), TestInstancedObject2->InstancedStruct.GetScriptStruct() == nullptr);
+
+				UTEST_TRUE(TEXT("TestContainerObject1.InstancedStructContainer struct"), TestContainerObject1->InstancedStructContainer[0].GetScriptStruct() == UserDefinedStruct);
+				UTEST_TRUE(TEXT("TestContainerObject1.Container.InstancedStructContainer struct"), TestContainerObject1->Container.InstancedStructContainer.Num() == 0);
+				UTEST_TRUE(TEXT("TestContainerObject2.InstancedStructContainer struct"), TestContainerObject2->Container.InstancedStructContainer[0].GetScriptStruct() == UserDefinedStruct);
+				UTEST_TRUE(TEXT("TestContainerObject2.Container.InstancedStructContainer struct"), TestContainerObject2->InstancedStructContainer.Num() == 0);
+
+				UTEST_TRUE(TEXT("TestBagObject1.InstancedPropertyBag struct"), TestBagObject1->InstancedPropertyBag.GetValue().GetScriptStruct() != UserDefinedStruct);
+				UTEST_TRUE(TEXT("TestBagObject1.Container.InstancedPropertyBag struct"), TestBagObject1->Container.InstancedPropertyBag.GetValue().GetScriptStruct() == nullptr);
+				UTEST_TRUE(TEXT("TestBagObject2.InstancedPropertyBag struct"), TestBagObject2->Container.InstancedPropertyBag.GetValue().GetScriptStruct() != UserDefinedStruct);
+				UTEST_TRUE(TEXT("TestBagObject2.Container.InstancedPropertyBag struct"), TestBagObject2->InstancedPropertyBag.GetValue().GetScriptStruct() == nullptr);
+
+				UTEST_TRUE(TEXT("TestInstancedObject1 B Value"), 22 == TryGetUDSPropertyValue_InContainer<int32>(UserDefinedStruct, BName, TestInstancedObject1->InstancedStruct.GetMemory()).Get(0));
+				UTEST_TRUE(TEXT("TestInstancedObject1 B Value"), !TryGetUDSPropertyValue_InContainer<int32>(UserDefinedStruct, BName, TestInstancedObject1->Container.InstancedStruct.GetMemory()).IsSet());
+				UTEST_TRUE(TEXT("TestInstancedObject2 B Value"), !TryGetUDSPropertyValue_InContainer<int32>(UserDefinedStruct, BName, TestInstancedObject2->InstancedStruct.GetMemory()).IsSet());
+				UTEST_TRUE(TEXT("TestInstancedObject2 B Value"), 22 == TryGetUDSPropertyValue_InContainer<int32>(UserDefinedStruct, BName, TestInstancedObject2->Container.InstancedStruct.GetMemory()).Get(0));
+
+				UTEST_TRUE(TEXT("TestContainerObject1 B Value"), 22 == TryGetUDSPropertyValue_InContainer<int32>(UserDefinedStruct, BName, TestContainerObject1->InstancedStructContainer[0].GetMemory()).Get(0));
+				UTEST_TRUE(TEXT("TestContainerObject2 B Value"), 22 == TryGetUDSPropertyValue_InContainer<int32>(UserDefinedStruct, BName, TestContainerObject2->Container.InstancedStructContainer[0].GetMemory()).Get(0));
+
+				UTEST_TRUE(TEXT("TestInstancedObject1 C Value"), 3 == TryGetUDSPropertyValue_InContainer<int32>(UserDefinedStruct, CName, TestInstancedObject1->InstancedStruct.GetMemory()).Get(0));
+				UTEST_TRUE(TEXT("TestInstancedObject1 C Value"), !TryGetUDSPropertyValue_InContainer<int32>(UserDefinedStruct, CName, TestInstancedObject1->Container.InstancedStruct.GetMemory()).IsSet());
+				UTEST_TRUE(TEXT("TestInstancedObject2 C Value"), !TryGetUDSPropertyValue_InContainer<int32>(UserDefinedStruct, CName, TestInstancedObject2->InstancedStruct.GetMemory()).IsSet());
+				UTEST_TRUE(TEXT("TestInstancedObject2 C Value"), 3 == TryGetUDSPropertyValue_InContainer<int32>(UserDefinedStruct, CName, TestInstancedObject2->Container.InstancedStruct.GetMemory()).Get(0));
+
+				UTEST_TRUE(TEXT("TestContainerObject1 C Value"), 3 == TryGetUDSPropertyValue_InContainer<int32>(UserDefinedStruct, CName, TestContainerObject1->InstancedStructContainer[0].GetMemory()).Get(0));
+				UTEST_TRUE(TEXT("TestContainerObject2 C Value"), 3 == TryGetUDSPropertyValue_InContainer<int32>(UserDefinedStruct, CName, TestContainerObject2->Container.InstancedStructContainer[0].GetMemory()).Get(0));
+
+				TValueOrError<FStructView, EPropertyBagResult> BagResult1 = TestBagObject1->InstancedPropertyBag.GetValueStruct("MyStruct1");
+				UTEST_TRUE(TEXT("Bag1Result1"), BagResult1.HasValue());
+				UTEST_TRUE(TEXT("BagResult1 struct"), BagResult1.GetValue().GetScriptStruct() == UserDefinedStruct);
+				UTEST_TRUE(TEXT("Bag1Result1 B Value"), 22 == TryGetUDSPropertyValue_InContainer<int32>(UserDefinedStruct, BName, BagResult1.GetValue().GetMemory()).Get(0));
+				UTEST_TRUE(TEXT("Bag1Result1 C Value"), 3 == TryGetUDSPropertyValue_InContainer<int32>(UserDefinedStruct, CName, BagResult1.GetValue().GetMemory()).Get(0));
+				TValueOrError<FStructView, EPropertyBagResult> BagResult2 = TestBagObject2->Container.InstancedPropertyBag.GetValueStruct("MyStruct2");
+				UTEST_TRUE(TEXT("Bag1Result2"), BagResult2.HasValue());
+				UTEST_TRUE(TEXT("BagResult2 struct"), BagResult2.GetValue().GetScriptStruct() == UserDefinedStruct);
+				UTEST_TRUE(TEXT("Bag1Result2 B Value"), 22 == TryGetUDSPropertyValue_InContainer<int32>(UserDefinedStruct, BName, BagResult2.GetValue().GetMemory()).Get(0));
+				UTEST_TRUE(TEXT("Bag1Result2 C Value"), 3 == TryGetUDSPropertyValue_InContainer<int32>(UserDefinedStruct, CName, BagResult2.GetValue().GetMemory()).Get(0));
+			}
+		}
+
+		return true;
+	}
+
+	IMPLEMENT_SIMPLE_AUTOMATION_TEST(FTest_UserDefinedStructReinstancerDefaultValues
+		, "System.StructUtils.Reinstancer.DefaultValues"
+		, EAutomationTestFlags::EditorContext | EAutomationTestFlags::EngineFilter);
+	bool FTest_UserDefinedStructReinstancerDefaultValues::RunTest(const FString& Parameters)
+	{
+		static const FName AName("A");
+		static const FName BName("B");
+		static const FName CName("C");
+
+		TStrongObjectPtr<UTestObjectThatContainsStruct> UDSContainer(NewObject<UTestObjectThatContainsStruct>());
+
+		// Create UDS
+		FInstancedPropertyBag BagA;
+		BagA.AddProperties({
+			{ AName, EPropertyBagPropertyType::Int32 },
+			{ BName, EPropertyBagPropertyType::Int32 },
+			{ CName, EPropertyBagPropertyType::Int32 },
+			});
+		BagA.SetValueInt32(AName, 1);
+		BagA.SetValueInt32(BName, 2);
+		BagA.SetValueInt32(CName, 3);
+
+		UUserDefinedStruct* UserDefinedStruct = nullptr;
+		{
+			TValueOrError<UUserDefinedStruct*, void> Result = UE::StructUtils::CreateUserDefinedStructFromDescs(UDSContainer.Get(), BagA, "UDS");
+
+			UTEST_TRUE("UUserDefinedStruct created", Result.HasValue());
+			UserDefinedStruct = Result.GetValue();
+			UTEST_TRUE("UUserDefinedStruct valid", UserDefinedStruct != nullptr);
+		}
+
+		// Assign to test object
+		TStrongObjectPtr<UTestObjectThatContainsStruct> TestInstancedObject(NewObject<UTestObjectThatContainsStruct>());
+		TestInstancedObject->InstancedStruct = UserDefinedStruct;
+		TestInstancedObject->Container.InstancedStruct = UserDefinedStruct;
+		TestInstancedObject->InstancedStructContainer.InsertAt(0, TArrayView<FInstancedStruct>(&TestInstancedObject->InstancedStruct, 1));
+		TestInstancedObject->Container.InstancedStructContainer.InsertAt(0, TArrayView<FInstancedStruct>(&TestInstancedObject->InstancedStruct, 1));
+		TestInstancedObject->InstancedPropertyBag.AddProperties({{"MyStruct1", EPropertyBagPropertyType::Struct, UserDefinedStruct}});
+		TestInstancedObject->Container.InstancedPropertyBag.AddProperties({ {"MyStruct2", EPropertyBagPropertyType::Struct, UserDefinedStruct} });
+
+		auto TestStruct = [this, TestInstancedObject, UserDefinedStruct]()
+			{
+				UTEST_TRUE(TEXT("TestInstancedObject.InstancedStruct"), TestInstancedObject->InstancedStruct.GetScriptStruct() == UserDefinedStruct);
+				UTEST_TRUE(TEXT("TestInstancedObject.Container.InstancedStruct"), TestInstancedObject->Container.InstancedStruct.GetScriptStruct() == UserDefinedStruct);
+				UTEST_TRUE(TEXT("TestInstancedObject.InstancedStructContainer"), TestInstancedObject->InstancedStructContainer[0].GetScriptStruct() == UserDefinedStruct);
+				UTEST_TRUE(TEXT("TestInstancedObject.Container.InstancedStructContainer struct"), TestInstancedObject->Container.InstancedStructContainer[0].GetScriptStruct() == UserDefinedStruct);
+				UTEST_TRUE(TEXT("TestInstancedObject.InstancedPropertyBag"), TestInstancedObject->InstancedPropertyBag.GetValue().GetScriptStruct() != UserDefinedStruct);
+				UTEST_TRUE(TEXT("TestInstancedObject.Container.InstancedPropertyBag"), TestInstancedObject->Container.InstancedPropertyBag.GetValue().GetScriptStruct() != UserDefinedStruct);
+
+				TValueOrError<FStructView, EPropertyBagResult> BagResult1 = TestInstancedObject->InstancedPropertyBag.GetValueStruct("MyStruct1");
+				UTEST_TRUE(TEXT("TestInstancedObject.InstancedPropertyBag"), BagResult1.HasValue());
+				UTEST_TRUE(TEXT("TestInstancedObject.InstancedPropertyBag.GetStruct"), BagResult1.GetValue().GetScriptStruct() == UserDefinedStruct);
+				TValueOrError<FStructView, EPropertyBagResult> BagResult2 = TestInstancedObject->Container.InstancedPropertyBag.GetValueStruct("MyStruct2");
+				UTEST_TRUE(TEXT("TestInstancedObject.Container.InstancedPropertyBag"), BagResult2.HasValue());
+				UTEST_TRUE(TEXT("TestInstancedObject.Container.InstancedPropertyBag.GetStruct"), BagResult2.GetValue().GetScriptStruct() == UserDefinedStruct);
+
+				return true;
+			};
+
+		auto TestValue = [this, TestInstancedObject, UserDefinedStruct](FName PropertyName, int32 Value)
+			{
+				UTEST_TRUE(TEXT("TestInstancedObject.InstancedStruct"), Value == TryGetUDSPropertyValue_InContainer<int32>(UserDefinedStruct, PropertyName, TestInstancedObject->InstancedStruct.GetMemory()).Get(0));
+				UTEST_TRUE(TEXT("TestInstancedObject.Container.InstancedStruct"), Value == TryGetUDSPropertyValue_InContainer<int32>(UserDefinedStruct, PropertyName, TestInstancedObject->Container.InstancedStruct.GetMemory()).Get(0));
+				UTEST_TRUE(TEXT("TestInstancedObject.InstancedStructContainer"), Value == TryGetUDSPropertyValue_InContainer<int32>(UserDefinedStruct, PropertyName, TestInstancedObject->InstancedStructContainer[0].GetMemory()).Get(0));
+				UTEST_TRUE(TEXT("TestInstancedObject.Container.InstancedStructContainer"), Value == TryGetUDSPropertyValue_InContainer<int32>(UserDefinedStruct, PropertyName, TestInstancedObject->Container.InstancedStructContainer[0].GetMemory()).Get(0));
+
+				TValueOrError<FStructView, EPropertyBagResult> BagResult1 = TestInstancedObject->InstancedPropertyBag.GetValueStruct("MyStruct1");
+				UTEST_TRUE(TEXT("TestInstancedObject.InstancedPropertyBag"), Value == TryGetUDSPropertyValue_InContainer<int32>(UserDefinedStruct, PropertyName, BagResult1.GetValue().GetMemory()).Get(0));
+				TValueOrError<FStructView, EPropertyBagResult> BagResult2 = TestInstancedObject->Container.InstancedPropertyBag.GetValueStruct("MyStruct2");
+				UTEST_TRUE(TEXT("TestInstancedObject.Container.InstancedPropertyBag"), Value == TryGetUDSPropertyValue_InContainer<int32>(UserDefinedStruct, PropertyName, BagResult2.GetValue().GetMemory()).Get(0));
+
+				return true;
+			};
+
+		auto SetValue = [this, TestInstancedObject, UserDefinedStruct](FName PropertyName, int32 Value)
+			{
+				UTEST_TRUE(TEXT("TestInstancedObject.InstancedStruct"), TrySetUDSPropertyValue_InContainer<int32>(UserDefinedStruct, PropertyName, TestInstancedObject->InstancedStruct.GetMutableMemory(), Value));
+				UTEST_TRUE(TEXT("TestInstancedObject.Container.InstancedStruct"), TrySetUDSPropertyValue_InContainer<int32>(UserDefinedStruct, PropertyName, TestInstancedObject->Container.InstancedStruct.GetMutableMemory(), Value));
+				UTEST_TRUE(TEXT("TestInstancedObject.InstancedStructContainer"), TrySetUDSPropertyValue_InContainer<int32>(UserDefinedStruct, PropertyName, TestInstancedObject->InstancedStructContainer[0].GetMemory(), Value));
+				UTEST_TRUE(TEXT("TestInstancedObject.Container.InstancedStructContainer struct"), TrySetUDSPropertyValue_InContainer<int32>(UserDefinedStruct, PropertyName, TestInstancedObject->Container.InstancedStructContainer[0].GetMemory(), Value));
+
+				TValueOrError<FStructView, EPropertyBagResult> BagResult1 = TestInstancedObject->InstancedPropertyBag.GetValueStruct("MyStruct1");
+				UTEST_TRUE(TEXT("TestInstancedObject.InstancedPropertyBag"), TrySetUDSPropertyValue_InContainer<int32>(UserDefinedStruct, PropertyName, BagResult1.GetValue().GetMemory(), Value));
+				TValueOrError<FStructView, EPropertyBagResult> BagResult2 = TestInstancedObject->Container.InstancedPropertyBag.GetValueStruct("MyStruct2");
+				UTEST_TRUE(TEXT("TestInstancedObject.Container.InstancedPropertyBag"), TrySetUDSPropertyValue_InContainer<int32>(UserDefinedStruct, PropertyName, BagResult2.GetValue().GetMemory(), Value));
+
+				return true;
+			};
+
+		UTEST_TRUE(TEXT("Initial Test Struct"), TestStruct());
+		UTEST_TRUE(TEXT("Initial Test Values"), TestValue(AName, 1));
+		UTEST_TRUE(TEXT("Initial Test Values"), TestValue(BName, 2));
+		UTEST_TRUE(TEXT("Initial Test Values"), TestValue(CName, 3));
+
+		{
+			BagA.SetValueInt32(BName, -2);
+			const bool bModified = UE::StructUtils::UpdateUserDefinedStructFromDescs(UserDefinedStruct, BagA);
+			UTEST_TRUE("UUserDefinedStruct modified", bModified);
+		}
+
+		UTEST_TRUE(TEXT("Initial Test Struct"), TestStruct());
+		UTEST_TRUE(TEXT("Initial Test Values"), TestValue(AName, 1));
+		UTEST_TRUE(TEXT("Initial Test Values"), TestValue(BName, -2));
+		UTEST_TRUE(TEXT("Initial Test Values"), TestValue(CName, 3));
+
+		// Set the value of C
+		{
+			UTEST_TRUE(TEXT("Initial Test Values"), SetValue(CName, 33));
+			UTEST_TRUE(TEXT("Initial Test Values"), TestValue(CName, 33));
+		}
+
+		// Set the default value of C. Should stay has it is. The child already modified it.
+		{
+			BagA.SetValueInt32(CName, -3);
+			const bool bModified = UE::StructUtils::UpdateUserDefinedStructFromDescs(UserDefinedStruct, BagA);
+			UTEST_TRUE("UUserDefinedStruct modified", bModified);
+		}
+
+		UTEST_TRUE(TEXT("Initial Test Struct"), TestStruct());
+		UTEST_TRUE(TEXT("Initial Test Values"), TestValue(AName, 1));
+		UTEST_TRUE(TEXT("Initial Test Values"), TestValue(BName, -2));
+		UTEST_TRUE(TEXT("Initial Test Values"), TestValue(CName, 33));
+
+		return true;
+	}
+
+	// A FInstancedStruct is an instance of a UPropertyBag.
+	IMPLEMENT_SIMPLE_AUTOMATION_TEST(FTest_PropertyBagReinstancer
+		, "System.StructUtils.Reinstancer.PropertyBag"
+		, EAutomationTestFlags::EditorContext | EAutomationTestFlags::EngineFilter);
+	bool FTest_PropertyBagReinstancer::RunTest(const FString& Parameters)
+	{
+		static const FName AName("A");
+		static const FName BName("B");
+		static const FName CName("C");
+		static const FName DName("D");
+
+		TStrongObjectPtr<UTestObjectThatContainsStruct> UDSContainer(NewObject<UTestObjectThatContainsStruct>());
+
+		FInstancedPropertyBag BagA;
+		BagA.AddProperties({ { CName, EPropertyBagPropertyType::Int32 } });
+		BagA.SetValueInt32(CName, 3);
+
+		// Create UDS
+		UUserDefinedStruct* UserDefinedStruct = nullptr;
+		{
+			TValueOrError<UUserDefinedStruct*, void> Result = UE::StructUtils::CreateUserDefinedStructFromDescs(UDSContainer.Get(), BagA, "UDS");
+
+			UTEST_TRUE("UUserDefinedStruct created", Result.HasValue());
+			UserDefinedStruct = Result.GetValue();
+			UTEST_TRUE("UUserDefinedStruct valid", UserDefinedStruct != nullptr);
+		}
+
+		// Create a property bag to keep it's reference
+		FInstancedPropertyBag BagB;
+		BagB.AddProperties({
+			{ AName, EPropertyBagPropertyType::Int32 },
+			{ BName, EPropertyBagPropertyType::Int32 },
+			{"MyStruct", EPropertyBagPropertyType::Struct, UserDefinedStruct}
+			});
+		BagB.SetValueInt32(AName, 1);
+		BagB.SetValueInt32(BName, 2);
+
+		// Assign the property bag generated struct to a test object
+		TStrongObjectPtr<UTestObjectThatContainsStruct> TestInstancedObject(NewObject<UTestObjectThatContainsStruct>());
+		TestInstancedObject->InstancedStruct = BagB.GetValue();
+		FConstStructView BagBValue = BagB.GetValue();
+		TestInstancedObject->InstancedStructContainer.InsertAt(0, TArrayView<FConstStructView>(&BagBValue, 1));
+
+		auto TestStruct = [this, TestInstancedObject](TNotNull<const UScriptStruct*> ScriptStruct, TNotNull<const UScriptStruct*> MyStruct)
+			{
+				UTEST_TRUE(TEXT("TestInstancedObject.InstancedStruct"), TestInstancedObject->InstancedStruct.GetScriptStruct() == ScriptStruct);
+				UTEST_TRUE(TEXT("TestInstancedObject.InstancedStructContainer"), TestInstancedObject->InstancedStructContainer[0].GetScriptStruct() == ScriptStruct);
+				const FProperty* Property = ScriptStruct->FindPropertyByName("MyStruct");
+				UTEST_TRUE(TEXT("Property"), Property != nullptr);
+				const FStructProperty* StructProperty = CastField<const FStructProperty>(Property);
+				UTEST_TRUE(TEXT("StructProperty"), StructProperty != nullptr);
+				if (StructProperty == nullptr) // C6011
+				{
+					return false;
+				}
+				UTEST_TRUE(TEXT("StructProperty == MyStruct"), MyStruct == StructProperty->Struct);
+				return true;
+			};
+
+		auto TestValue = [this, TestInstancedObject](FName PropertyName, int32 Value)
+			{
+				const UScriptStruct* ScriptStruct = TestInstancedObject->InstancedStruct.GetScriptStruct();
+				UTEST_TRUE(TEXT("TestInstancedObject.InstancedStruct"), Value == TryGetPropertyValue_InContainer<int32>(ScriptStruct, PropertyName, TestInstancedObject->InstancedStruct.GetMemory()).Get(0));
+				UTEST_TRUE(TEXT("TestInstancedObject.Container.InstancedStruct"), Value == TryGetPropertyValue_InContainer<int32>(ScriptStruct, PropertyName, TestInstancedObject->InstancedStructContainer[0].GetMemory()).Get(0));
+				return true;
+			};
+
+		auto TestPropertyCDValue = [this, TestInstancedObject](int32 Value, int32 Offset)
+			{
+				const UScriptStruct* ScriptStruct = TestInstancedObject->InstancedStruct.GetScriptStruct();
+				const FStructProperty* StructProperty = CastField<const FStructProperty>(ScriptStruct->FindPropertyByName("MyStruct"));
+				if (StructProperty == nullptr) // C6011
+				{
+					return false;
+				}
+
+				int32 CStructValue1 = 0;
+				StructProperty->GetValue_InContainer(TestInstancedObject->InstancedStruct.GetMemory() + Offset, &CStructValue1);
+				int32 CStructValue2 = 0;
+				StructProperty->GetValue_InContainer(TestInstancedObject->InstancedStructContainer[0].GetMemory() + Offset, &CStructValue2);
+				UTEST_TRUE(TEXT("TestInstancedObject.InstancedStruct"), Value == CStructValue1);
+				UTEST_TRUE(TEXT("TestInstancedObject.Container.InstancedStruct"), Value == CStructValue2);
+				return true;
+			};
+
+		UTEST_TRUE(TEXT("Initial Test Struct"), TestStruct(BagB.GetValue().GetScriptStruct(), UserDefinedStruct));
+		UTEST_TRUE(TEXT("Initial Test Values"), TestValue(AName, 1));
+		UTEST_TRUE(TEXT("Initial Test Values"), TestValue(BName, 2));
+		UTEST_TRUE(TEXT("Initial Test Values"), TestPropertyCDValue(3, 0));
+
+		// Set the value of C.
+		// TestInstancedObject->InstancedStruct is not a FInstancedPropertyBag. It wrongly point to the temporary UPropertyBag struct.
+		//This is not supported but we do not want to crash on dangling memory.
+		{
+			UTEST_TRUE(TEXT("Assign"), EPropertyBagResult::Success == BagA.SetValueInt32(CName, -3));
+			BagA.AddProperties({ { DName, EPropertyBagPropertyType::Int32 } });
+			UTEST_TRUE(TEXT("Assign"), EPropertyBagResult::Success == BagA.SetValueInt32(DName, 4));
+			const bool bModified = UE::StructUtils::UpdateUserDefinedStructFromDescs(UserDefinedStruct, BagA);
+			UTEST_TRUE(TEXT("UUserDefinedStruct modified"), bModified);
+		}
+		{
+			const UScriptStruct* OldBagBStruct = BagB.GetValue().GetScriptStruct();
+			UTEST_TRUE(TEXT("Modified Test Struct A"), OldBagBStruct->FindPropertyByName(AName) != nullptr);
+			UTEST_TRUE(TEXT("Modified Test Struct B"), OldBagBStruct->FindPropertyByName(BName) != nullptr);
+			const FProperty* MyStructProperty = OldBagBStruct->FindPropertyByName("MyStruct");
+			UTEST_TRUE(TEXT("Modified Test Struct Struct"), MyStructProperty != nullptr);
+			const FStructProperty* MyStructStructProperty = CastField<const FStructProperty>(MyStructProperty);
+			UTEST_TRUE(TEXT("Modified Test Struct Struct"), MyStructStructProperty != nullptr);
+			if (MyStructStructProperty == nullptr) // C6011
+			{
+				return false;
+			}
+			UTEST_TRUE(TEXT("Modified Test Struct Struct"), MyStructStructProperty->Struct != nullptr);
+			UTEST_TRUE(TEXT("Modified Test Struct Newer"), (MyStructStructProperty->Struct->StructFlags & STRUCT_NewerVersionExists) != 0);
+		}
+
+		{
+			const UScriptStruct* NewBagBStruct = TestInstancedObject->InstancedStruct.GetScriptStruct();
+			UTEST_TRUE(TEXT("Modified Test Struct A"), NewBagBStruct->FindPropertyByName(AName) != nullptr);
+			UTEST_TRUE(TEXT("Modified Test Struct B"), NewBagBStruct->FindPropertyByName(BName) != nullptr);
+			const FProperty* MyStructProperty = NewBagBStruct->FindPropertyByName("MyStruct");
+			UTEST_TRUE(TEXT("Modified Test Struct Struct"), MyStructProperty != nullptr);
+			const FStructProperty* MyStructStructProperty = CastField<const FStructProperty>(MyStructProperty);
+			UTEST_TRUE(TEXT("Modified Test Struct Struct"), MyStructStructProperty != nullptr);
+			if (MyStructStructProperty == nullptr) // C6011
+			{
+				return false;
+			}
+			UTEST_TRUE(TEXT("Modified Test Struct Struct"), MyStructStructProperty->Struct != nullptr);
+			UTEST_TRUE(TEXT("Modified Test Struct Newer"), (MyStructStructProperty->Struct->StructFlags & STRUCT_NewerVersionExists) == 0);
+
+			const UUserDefinedStruct* NewUDS = Cast<UUserDefinedStruct>(MyStructStructProperty->Struct);
+			UTEST_TRUE(TEXT("Modified is UDS"), NewUDS != nullptr);
+			if (NewUDS == nullptr) // C6011
+			{
+				return false;
+			}
+			const UUserDefinedStructEditorData* EditorData = Cast<UUserDefinedStructEditorData>(NewUDS->EditorData);
+			UTEST_TRUE(TEXT("Modified is UDS editor data"), EditorData != nullptr);
+			if (EditorData == nullptr) // C6011
+			{
+				return false;
+			}
+			UTEST_TRUE(TEXT("Modified Find Struct C"), EditorData->FindProperty(NewUDS, CName) != nullptr);
+			UTEST_TRUE(TEXT("Modified Find Struct D"), EditorData->FindProperty(NewUDS, DName) != nullptr);
+			UTEST_TRUE(TEXT("Modified Test Values"), TestValue(AName, 1));
+			UTEST_TRUE(TEXT("Modified Test Values"), TestValue(BName, 2));
+			UTEST_TRUE(TEXT("Modified Test Values"), TestPropertyCDValue(-3, 0));
+			UTEST_TRUE(TEXT("Modified Test Values"), TestPropertyCDValue(4, sizeof(int32)));
+		}
+
+		return true;
+	}
+
+} // UE::StructUtils::Private
+
+#undef LOCTEXT_NAMESPACE
diff --git a/Engine/Source/Editor/StructUtilsEditor/Private/Tests/StructReinstancerTests.h b/Engine/Source/Editor/StructUtilsEditor/Private/Tests/StructReinstancerTests.h
new file mode 100644
--- /dev/null
+++ b/Engine/Source/Editor/StructUtilsEditor/Private/Tests/StructReinstancerTests.h
@@ -0,0 +1,60 @@
+// Copyright Epic Games, Inc. All Rights Reserved.
+
+#pragma once
+
+#include "CoreMinimal.h"
+#include "Templates/SubclassOf.h"
+#include "StructUtils/InstancedStruct.h"
+#include "StructUtils/InstancedStructContainer.h"
+#include "StructUtils/PropertyBag.h"
+#include "StructUtils/UserDefinedStruct.h"
+#include "StructReinstancerTests.generated.h"
+
+namespace UE::StructUtils::Private
+{
+USTRUCT()
+struct FTestStructThatContainsStruct
+{
+	GENERATED_BODY()
+
+	UPROPERTY()
+	FVector Vect;
+
+	UPROPERTY()
+	FInstancedStruct InstancedStruct;
+
+	UPROPERTY()
+	FInstancedStructContainer InstancedStructContainer;
+
+	UPROPERTY()
+	FInstancedPropertyBag InstancedPropertyBag;
+};
+
+UCLASS()
+class UTestObjectThatContainsStructBase : public UObject
+{
+	GENERATED_BODY()
+
+public:
+	UPROPERTY()
+	FVector Vect;
+
+	UPROPERTY()
+	FTestStructThatContainsStruct Container;
+
+	UPROPERTY()
+	FInstancedStruct InstancedStruct;
+
+	UPROPERTY()
+	FInstancedStructContainer InstancedStructContainer;
+
+	UPROPERTY()
+	FInstancedPropertyBag InstancedPropertyBag;
+};
+
+UCLASS()
+class UTestObjectThatContainsStruct : public UTestObjectThatContainsStructBase
+{
+	GENERATED_BODY()
+};
+}
\ No newline at end of file
diff --git a/Engine/Source/Editor/StructUtilsEditor/Public/StructUtilsEditorUtils.h b/Engine/Source/Editor/StructUtilsEditor/Public/StructUtilsEditorUtils.h
--- a/Engine/Source/Editor/StructUtilsEditor/Public/StructUtilsEditorUtils.h
+++ b/Engine/Source/Editor/StructUtilsEditor/Public/StructUtilsEditorUtils.h
@@ -18,7 +18,7 @@ namespace UE::StructUtils
 	 * Fails if the name is already used by another asset or if the asset could not be created.
 	 * When the property bag is empty, no user defined struct will be created.
 	 */
-	UE_API TValueOrError<UUserDefinedStruct*, void> CreateUserDefinedStructFromDescs(TNotNull<UObject*> PropertyBagOuter, const FInstancedPropertyBag& PropertyBag, FName UserDefinedStructName);
+	UE_API TValueOrError<UUserDefinedStruct*, void> CreateUserDefinedStructFromDescs(TNotNull<UObject*> UserDefinedStructOuter, const FInstancedPropertyBag& PropertyBag, FName UserDefinedStructName);
 
 	/**
 	 * Update, if needed, an existing user defined struct to have the exact same properties and default value as the property bag.
diff --git a/Engine/Source/Editor/StructUtilsEditor/StructUtilsEditor.Build.cs b/Engine/Source/Editor/StructUtilsEditor/StructUtilsEditor.Build.cs
--- a/Engine/Source/Editor/StructUtilsEditor/StructUtilsEditor.Build.cs
+++ b/Engine/Source/Editor/StructUtilsEditor/StructUtilsEditor.Build.cs
@@ -8,6 +8,7 @@ namespace UnrealBuildTool.Rules
 		public StructUtilsEditor(ReadOnlyTargetRules Target) : base(Target)
 		{
 			PCHUsage = ModuleRules.PCHUsageMode.UseExplicitOrSharedPCHs;
+			bAllowUETypesInNamespaces = true;
 
 			PublicIncludePaths.AddRange(
 				new string[] {
diff --git a/Engine/Source/Runtime/CoreUObject/Private/StructUtils/InstancedStruct.cpp b/Engine/Source/Runtime/CoreUObject/Private/StructUtils/InstancedStruct.cpp
--- a/Engine/Source/Runtime/CoreUObject/Private/StructUtils/InstancedStruct.cpp
+++ b/Engine/Source/Runtime/CoreUObject/Private/StructUtils/InstancedStruct.cpp
@@ -9,10 +9,7 @@
 #include "UObject/CoreRedirects.h"
 
 #if WITH_EDITOR
-#include "Serialization/MemoryReader.h"
-#include "Serialization/MemoryWriter.h"
-#include "Serialization/ObjectAndNameAsStringProxyArchive.h"
-#include "StructUtils/UserDefinedStruct.h"
+#include "StructUtils/StructReinstancer.h"
 #endif // WITH_EDITOR
 
 #include UE_INLINE_GENERATED_CPP_BY_NAME(InstancedStruct)
@@ -211,10 +208,23 @@ bool FInstancedStruct::Serialize(FArchive& Ar, UStruct* DefaultsStruct, const vo
 		Ar << SerializedScriptStruct;
 		if (SerializedScriptStruct)
 		{
+#if WITH_EDITOR
+			// During reinstanciation, the serialized struct can be the the duplicated struct.
+			//Replace it with the new version.
+			if (const UE::StructUtils::FStructReinstancer* Reinstancer = UE::StructUtils::FStructReinstancer::GetInstance();
+				Reinstancer && Reinstancer->IsSerializingForReinstanciation())
+			{
+				if (const UScriptStruct* CompiledStruct = Reinstancer->GetCompiledReinstantingStruct(SerializedScriptStruct))
+				{
+					SerializedScriptStruct = const_cast<UScriptStruct*>(CompiledStruct);
+				}
+			}
+#endif
+
 			Ar.Preload(SerializedScriptStruct);
 		}
 
-		// Destroy the exiting memory and create a new one using the struct default values.
+		// Create a new one using the struct default values.
 		if (DefaultInstancedStruct && SerializedScriptStruct && SerializedScriptStruct->IsChildOf(DefaultInstancedStruct->GetScriptStruct()))
 		{
 			InitializeAs(SerializedScriptStruct, DefaultInstancedStruct->GetMemory());
@@ -249,22 +259,7 @@ bool FInstancedStruct::Serialize(FArchive& Ar, UStruct* DefaultsStruct, const vo
 	else if (Ar.IsSaving())
 	{
 		// UScriptStruct type
-#if WITH_EDITOR
-		const UUserDefinedStruct* UserDefinedStruct = Cast<UUserDefinedStruct>(ScriptStruct);
-		if (UserDefinedStruct
-			&& UserDefinedStruct->Status == EUserDefinedStructureStatus::UDSS_Duplicate
-			&& UserDefinedStruct->PrimaryStruct.IsValid())
-		{
-			// If saving a duplicated UDS, save the primary type instead, so that the data is loaded with the original struct.
-			// This is used as part of the user defined struct reinstancing logic.
-			UUserDefinedStruct* PrimaryUserDefinedStruct = UserDefinedStruct->PrimaryStruct.Get();
-			Ar << PrimaryUserDefinedStruct;
-		}
-		else
-#endif // WITH_EDITOR
-		{
-			Ar << ScriptStruct;
-		}
+		Ar << ScriptStruct;
 	
 		// Size of the serialized memory (reserve location)
 		const int64 SizeOffset = Ar.Tell(); // Position to write the actual size after struct serialization
@@ -484,6 +479,18 @@ bool FInstancedStruct::Identical(const FInstancedStruct* Other, uint32 PortFlags
 
 	if (StructTypePtr)
 	{
+#if WITH_EDITOR
+		// Struct can't be identical when re-instancing. The struct has to serialize to get the new values.
+		if (const UE::StructUtils::FStructReinstancer* Reinstancer = UE::StructUtils::FStructReinstancer::GetInstance())
+		{
+			if (Reinstancer->IsSerializingForReinstanciation()
+				&& Reinstancer->IsReinstanting(StructTypePtr))
+			{
+				return false;
+			}
+		}
+#endif
+
 		return StructTypePtr->CompareScriptStruct(GetMemory(), Other->GetMemory(), PortFlags);
 	}
 
@@ -494,46 +501,15 @@ void FInstancedStruct::AddStructReferencedObjects(FReferenceCollector& Collector
 {
 #if WITH_EDITOR
 	// Reference collector is used to visit all instances of instanced structs and replace their contents.
-	if (const UUserDefinedStruct* StructureToReinstance = UE::StructUtils::Private::GetStructureToReinstantiate())
+	if (UE::StructUtils::FStructReinstancer* Reinstancer = UE::StructUtils::FStructReinstancer::GetInstance())
 	{
-		if (const UUserDefinedStruct* UserDefinedStruct = Cast<UUserDefinedStruct>(ScriptStruct))
+		if (ScriptStruct && Reinstancer->IsReinstanting(ScriptStruct))
 		{
-			if (StructureToReinstance->Status == EUserDefinedStructureStatus::UDSS_Duplicate)
-			{
-				// On the first pass we replace the UDS with a duplicate that represents the currently allocated struct.
-				// GStructureToReinstance is the duplicated struct, and StructureToReinstance->PrimaryStruct is the UDS that is being reinstanced.
-				
-				if (UserDefinedStruct == StructureToReinstance->PrimaryStruct)
-				{
-					ScriptStruct = StructureToReinstance;
-				}
-			}
-			else
+			Reinstancer->MarkAsRequiresReinstanciation();
+			// It is not necessary to change the struct to a duplicated one, but it helps with tracking mistakes and prevents memory stomps.
+			if (const UScriptStruct* Duplicated = Reinstancer->GetDuplicatedReinstantingStruct(ScriptStruct))
 			{
-				// On the second pass we reinstantiate the data using serialization.
-				// When saving, the UDSs are written using the duplicate which represents current layout, but PrimaryStruct is serialized as the type.
-				// When reading, the data is initialized with the new type, and the serialization will take care of reading from the old data.
-
-				if (UserDefinedStruct->PrimaryStruct == StructureToReinstance)
-				{
-					if (UObject* Outer = UE::StructUtils::Private::GetCurrentReinstantiationOuterObject())
-					{
-						if (!Outer->IsA<UClass>() && !Outer->HasAnyFlags(RF_ClassDefaultObject))
-						{
-							(void)Outer->MarkPackageDirty();
-						}
-					}
-
-					TArray<uint8> Data;
-					
-					FMemoryWriter Writer(Data);
-					FObjectAndNameAsStringProxyArchive WriterProxy(Writer, /*bInLoadIfFindFails*/true);
-					Serialize(WriterProxy);
-
-					FMemoryReader Reader(Data);
-					FObjectAndNameAsStringProxyArchive ReaderProxy(Reader, /*bInLoadIfFindFails*/true);
-					Serialize(ReaderProxy);
-				}
+				ScriptStruct = Duplicated;
 			}
 		}
 	}
diff --git a/Engine/Source/Runtime/CoreUObject/Private/StructUtils/InstancedStructContainer.cpp b/Engine/Source/Runtime/CoreUObject/Private/StructUtils/InstancedStructContainer.cpp
--- a/Engine/Source/Runtime/CoreUObject/Private/StructUtils/InstancedStructContainer.cpp
+++ b/Engine/Source/Runtime/CoreUObject/Private/StructUtils/InstancedStructContainer.cpp
@@ -4,10 +4,7 @@
 #include "UObject/Linker.h"
 
 #if WITH_EDITOR
-#include "Serialization/MemoryReader.h"
-#include "Serialization/MemoryWriter.h"
-#include "Serialization/ObjectAndNameAsStringProxyArchive.h"
-#include "StructUtils/UserDefinedStruct.h"
+#include "StructUtils/StructReinstancer.h"
 #endif // WITH_EDITOR
 
 #include UE_INLINE_GENERATED_CPP_BY_NAME(InstancedStructContainer)
@@ -379,65 +376,22 @@ void FInstancedStructContainer::Empty()
 void FInstancedStructContainer::AddStructReferencedObjects(FReferenceCollector& Collector)
 {
 #if WITH_EDITOR
-	if (const UUserDefinedStruct* StructureToReinstantiate = UE::StructUtils::Private::GetStructureToReinstantiate())
+	// Reference collector is used to visit all instances of instanced structs and replace their contents.
+	if (UE::StructUtils::FStructReinstancer* Reinstancer = UE::StructUtils::FStructReinstancer::GetInstance())
 	{
-		bool bContainsUserDefineStruct = false;
-		for (int32 Index = 0, Num = NumItems; Index < Num; Index++)
+		for (int32 Index = 0; Index < NumItems; ++Index)
 		{
-			FItem& Item = GetItem(Index);
-			if (const UUserDefinedStruct* UserDefinedStruct = Cast<UUserDefinedStruct>(Item.ScriptStruct))
+			TObjectPtr<const UScriptStruct>& ScriptStruct = GetItem(Index).ScriptStruct;
+			if (ScriptStruct && Reinstancer->IsReinstanting(ScriptStruct))
 			{
-				if (UserDefinedStruct == StructureToReinstantiate
-					|| UserDefinedStruct == StructureToReinstantiate->PrimaryStruct
-					|| UserDefinedStruct->PrimaryStruct == StructureToReinstantiate
-					|| UserDefinedStruct->PrimaryStruct == StructureToReinstantiate->PrimaryStruct)
+				Reinstancer->MarkAsRequiresReinstanciation();
+				// It is not necessary to change the struct to a duplicated one, but it helps with tracking mistakes and prevents memory stomps.
+				if (const UScriptStruct* Duplicated = Reinstancer->GetDuplicatedReinstantingStruct(ScriptStruct))
 				{
-					bContainsUserDefineStruct = true;
-					break;
+					ScriptStruct = Duplicated;
 				}
 			}
 		}
-		
-		if (bContainsUserDefineStruct)
-		{
-			if (StructureToReinstantiate->Status == EUserDefinedStructureStatus::UDSS_Duplicate)
-			{
-				// On the first pass we replace the UDS with a duplicate that represents the currently allocated struct.
-				// StructureToReinstantiate is the duplicated struct, and StructureToReinstantiate->PrimaryStruct is the UDS that is being reinstantiated.
-				for (int32 Index = 0, Num = NumItems; Index < Num; Index++)
-				{
-					FItem& Item = GetItem(Index);
-					if (Item.ScriptStruct == StructureToReinstantiate->PrimaryStruct)
-					{
-						Item.ScriptStruct = StructureToReinstantiate;
-					}
-				}
-			}
-			else
-			{
-				// On the second pass we reinstantiate the data using serialization.
-				// When saving, the UDSs are written using the duplicate which represents current layout, but PrimaryStruct is serialized as the type.
-				// When reading, the data is initialized with the new type, and the serialization will take care of reading from the old data.
-				
-				if (UObject* Outer = UE::StructUtils::Private::GetCurrentReinstantiationOuterObject())
-				{
-					if (!Outer->IsA<UClass>() && !Outer->HasAnyFlags(RF_ClassDefaultObject))
-					{
-						(void)Outer->MarkPackageDirty();
-					}
-				}
-				
-				TArray<uint8> Data;
-
-				FMemoryWriter Writer(Data);
-				FObjectAndNameAsStringProxyArchive WriterProxy(Writer, /*bInLoadIfFindFails*/true);
-				Serialize(WriterProxy);
-
-				FMemoryReader Reader(Data);
-				FObjectAndNameAsStringProxyArchive ReaderProxy(Reader, /*bInLoadIfFindFails*/true);
-				Serialize(ReaderProxy);
-			}
-		}
 	}
 #endif // WITH_EDITOR
 	
@@ -479,6 +433,10 @@ bool FInstancedStructContainer::Identical(const FInstancedStructContainer* Other
 
 	bool bResult = true;
 
+#if WITH_EDITOR
+	const UE::StructUtils::FStructReinstancer* Reinstancer = UE::StructUtils::FStructReinstancer::GetInstance();
+#endif
+
 	// Check that the struct contents are identical.
 	for (int32 Index = 0; Index < Num(); Index++)
 	{ 
@@ -491,6 +449,16 @@ bool FInstancedStructContainer::Identical(const FInstancedStructContainer* Other
 		}
 		if (Item.ScriptStruct != nullptr && OtherItem.ScriptStruct != nullptr)
 		{
+#if WITH_EDITOR
+			// Struct can't be identical when re-instancing. The struct has to serialize to get the new values.
+			if (Reinstancer
+				&& Reinstancer->IsSerializingForReinstanciation()
+				&& Reinstancer->IsReinstanting(Item.ScriptStruct))
+			{
+				bResult = false;
+				break;
+			}
+#endif
 			const uint8* ItemMemory = Memory + Item.Offset;
 			const uint8* OtherItemMemory = Other->Memory + OtherItem.Offset;
 			
@@ -599,6 +567,10 @@ bool FInstancedStructContainer::Serialize(FArchive& Ar, UStruct* DefaultsStruct,
 		}
 		else if (Ar.IsLoading())
 		{
+#if WITH_EDITOR
+			const UE::StructUtils::FStructReinstancer* Reinstancer = UE::StructUtils::FStructReinstancer::GetInstance();
+#endif
+
 			// Load item types 
 			TArray<FConstStructView> Views;
 			Views.Reserve(NumItemsSerialized);
@@ -608,6 +580,17 @@ bool FInstancedStructContainer::Serialize(FArchive& Ar, UStruct* DefaultsStruct,
 				Ar << SerializedScriptStruct;
 				if (SerializedScriptStruct)
 				{
+#if WITH_EDITOR
+					// During reinstanciation, the serialized struct can be the the duplicated struct.
+					//Replace it with the new version.
+					if (Reinstancer && Reinstancer->IsSerializingForReinstanciation())
+					{
+						if (const UScriptStruct* CompiledStruct = Reinstancer->GetCompiledReinstantingStruct(SerializedScriptStruct))
+						{
+							SerializedScriptStruct = const_cast<UScriptStruct*>(CompiledStruct);
+						}
+					}
+#endif
 					Ar.Preload(SerializedScriptStruct);
 				}
 				Views.Add(FConstStructView(SerializedScriptStruct, nullptr));
@@ -651,22 +634,7 @@ bool FInstancedStructContainer::Serialize(FArchive& Ar, UStruct* DefaultsStruct,
 			{
 				const FItem& Item = GetItem(Index);
 				
-#if WITH_EDITOR
-				UUserDefinedStruct* UserDefinedStruct = Cast<UUserDefinedStruct>(ConstCast(Item.ScriptStruct));
-				if (UserDefinedStruct
-					&& UserDefinedStruct->Status == EUserDefinedStructureStatus::UDSS_Duplicate
-					&& UserDefinedStruct->PrimaryStruct.IsValid())
-				{
-					// If saving a duplicated UDS, save the primary type instead, so that the data is loaded with the original struct.
-					// This is used as part of the user defined struct reinstancing logic.
-					UUserDefinedStruct* PrimaryUserDefinedStruct = UserDefinedStruct->PrimaryStruct.Get(); 
-					Ar << PrimaryUserDefinedStruct;
-				}
-				else
-#endif // WITH_EDITOR
-				{
-					Ar << ConstCast(Item.ScriptStruct);
-				}
+				Ar << ConstCast(Item.ScriptStruct);
 			}
 
 			// Save item values
diff --git a/Engine/Source/Runtime/CoreUObject/Private/StructUtils/PropertyBag.cpp b/Engine/Source/Runtime/CoreUObject/Private/StructUtils/PropertyBag.cpp
--- a/Engine/Source/Runtime/CoreUObject/Private/StructUtils/PropertyBag.cpp
+++ b/Engine/Source/Runtime/CoreUObject/Private/StructUtils/PropertyBag.cpp
@@ -13,9 +13,7 @@
 #include "UObject/TextProperty.h"
 
 #if WITH_EDITOR
-#include "Serialization/MemoryReader.h"
-#include "Serialization/MemoryWriter.h"
-#include "Serialization/ObjectAndNameAsStringProxyArchive.h"
+#include "StructUtils/StructReinstancer.h"
 #include "StructUtils/UserDefinedStruct.h"
 #endif // WITH_EDITOR
 
@@ -3458,6 +3456,39 @@ TValueOrError<const FPropertyBagMapRef, EPropertyBagResult> FInstancedPropertyBa
 
 bool FInstancedPropertyBag::Identical(const FInstancedPropertyBag* Other, const uint32 PortFlags) const
 {
+#if WITH_EDITOR
+	// Struct can't be identical when re-instancing. The struct has to serialize to get the new values.
+	UE::StructUtils::FStructReinstancer* Reinstancer = UE::StructUtils::FStructReinstancer::GetInstance();
+	if (Other && Reinstancer && Reinstancer->IsSerializingForReinstanciation())
+	{
+		if (const UPropertyBag* BagStruct = GetPropertyBagStruct())
+		{
+			auto IsReinstanting = [Reinstancer](const UScriptStruct* Struct)
+				{
+					return Struct && Reinstancer->IsReinstanting(Struct);
+				};
+			TArrayView<const FPropertyBagPropertyDesc> PropertyDescs = BagStruct->PropertyDescs;
+			for (const FPropertyBagPropertyDesc& Desc : PropertyDescs)
+			{
+				if (Desc.ValueType == EPropertyBagPropertyType::Struct)
+				{
+					if (IsReinstanting(Cast<const UScriptStruct>(Desc.ValueTypeObject)))
+					{
+						return false;
+					}
+				}
+				if (Desc.KeyType == EPropertyBagPropertyType::Struct)
+				{
+					if (IsReinstanting(Cast<const UScriptStruct>(Desc.KeyTypeObject)))
+					{
+						return false;
+					}
+				}
+			}
+		}
+	}
+#endif // WITH_EDITOR
+
 	return Other && Value.Identical(&Other->Value, PortFlags);
 }
 
@@ -3494,18 +3525,48 @@ bool FInstancedPropertyBag::Serialize(FArchive& Ar)
 			TArray<FPropertyBagPropertyDesc> PropertyDescs;
 			Ar << PropertyDescs;
 
+#if WITH_EDITOR
+			const UE::StructUtils::FStructReinstancer* Reinstancer = UE::StructUtils::FStructReinstancer::GetInstance();
+#endif
+
 			for (FPropertyBagPropertyDesc& PropDesc : PropertyDescs)
 			{
 				if (PropDesc.ValueTypeObject)
 				{
+#if WITH_EDITOR
+					if (Reinstancer && PropDesc.ValueType == EPropertyBagPropertyType::Struct)
+					{
+						if (const UScriptStruct* ValueStruct = Cast<const UScriptStruct>(PropDesc.ValueTypeObject);
+							ValueStruct && Reinstancer->IsSerializingForReinstanciation())
+						{
+							if (const UScriptStruct* CompiledStruct = Reinstancer->GetCompiledReinstantingStruct(ValueStruct))
+							{
+								PropDesc.ValueTypeObject = CompiledStruct;
+							}
+						}
+					}
+#endif
 					Ar.Preload(const_cast<UObject*>(PropDesc.ValueTypeObject.Get()));
 				}
 				if (PropDesc.KeyTypeObject)
 				{
+#if WITH_EDITOR
+					if (Reinstancer && PropDesc.KeyType == EPropertyBagPropertyType::Struct)
+					{
+						if (const UScriptStruct* ValueStruct = Cast<const UScriptStruct>(PropDesc.KeyTypeObject);
+							ValueStruct && Reinstancer->IsSerializingForReinstanciation())
+						{
+							if (const UScriptStruct* CompiledStruct = Reinstancer->GetCompiledReinstantingStruct(ValueStruct))
+							{
+								PropDesc.KeyTypeObject = CompiledStruct;
+							}
+						}
+					}
+#endif
 					Ar.Preload(const_cast<UObject*>(PropDesc.KeyTypeObject.Get()));
 				}
 			}
-			
+
 			BagStruct = const_cast<UPropertyBag*>(UPropertyBag::GetOrCreateFromDescs(PropertyDescs));
 			Value.InitializeAs(BagStruct);
 
@@ -3532,34 +3593,6 @@ bool FInstancedPropertyBag::Serialize(FArchive& Ar)
 			check(BagStruct);
 			
 			TArray<FPropertyBagPropertyDesc> PropertyDescs = BagStruct->PropertyDescs;
-#if WITH_EDITOR
-			// Save primary struct for user defined struct properties.
-			// This is used as part of the user defined struct reinstancing logic.
-			for (FPropertyBagPropertyDesc& Desc : PropertyDescs)
-			{
-				if (Desc.ValueType == EPropertyBagPropertyType::Struct)
-				{
-					const UUserDefinedStruct* UserDefinedStruct = Cast<UUserDefinedStruct>(Desc.ValueTypeObject);
-					if (UserDefinedStruct
-						&& UserDefinedStruct->Status == EUserDefinedStructureStatus::UDSS_Duplicate
-						&& UserDefinedStruct->PrimaryStruct.IsValid())
-					{
-						Desc.ValueTypeObject = UserDefinedStruct->PrimaryStruct.Get();
-					}
-				}
-				if (Desc.KeyType == EPropertyBagPropertyType::Struct)
-				{
-					const UUserDefinedStruct* UserDefinedStruct = Cast<UUserDefinedStruct>(Desc.KeyTypeObject);
-					if (UserDefinedStruct
-						&& UserDefinedStruct->Status == EUserDefinedStructureStatus::UDSS_Duplicate
-						&& UserDefinedStruct->PrimaryStruct.IsValid())
-					{
-						Desc.KeyTypeObject = UserDefinedStruct->PrimaryStruct.Get();
-					}
-				}
-			}
-#endif // WITH_EDITOR
-
 			Ar << PropertyDescs;
 
 			const int64 SizeOffset = Ar.Tell(); // Position to write the actual size after struct serialization
@@ -3579,6 +3612,37 @@ bool FInstancedPropertyBag::Serialize(FArchive& Ar)
 			SerialSize = static_cast<int32>(FinalOffset - InitialOffset);
 			Ar << SerialSize;
 			Ar.Seek(FinalOffset);	// Reset archive to its position
+
+#if WITH_EDITOR
+			// Destroy the memory. Making sure we are not using the memory with an invalid struct.
+			if (UE::StructUtils::FStructReinstancer* Reinstancer = UE::StructUtils::FStructReinstancer::GetInstance())
+			{
+				if (Ar.IsSaving() && Reinstancer->IsSerializingForReinstanciation())
+				{
+					for (const FPropertyBagPropertyDesc& Desc : PropertyDescs)
+					{
+						if (Desc.ValueType == EPropertyBagPropertyType::Struct)
+						{
+							if (const UScriptStruct* ValueStruct = Cast<const UScriptStruct>(Desc.ValueTypeObject);
+								ValueStruct && Reinstancer->IsReinstanting(ValueStruct))
+							{
+								Value.InitializeAs(nullptr);
+								break;
+							}
+						}
+						if (Desc.KeyType == EPropertyBagPropertyType::Struct)
+						{
+							if (const UScriptStruct* KeyStruct = Cast<const UScriptStruct>(Desc.KeyTypeObject);
+								KeyStruct&& Reinstancer->IsReinstanting(KeyStruct))
+							{
+								Value.InitializeAs(nullptr);
+								break;
+							}
+						}
+					}
+				}
+			}
+#endif
 		}
 	}
 	else
@@ -3597,71 +3661,34 @@ void FInstancedPropertyBag::AddStructReferencedObjects(FReferenceCollector& Coll
 {
 #if WITH_EDITOR
 	// Reference collector is used to visit all instances of instanced structs the the like when a user defined struct is reinstanced.
-	if (const UUserDefinedStruct* StructureToReinstance = UE::StructUtils::Private::GetStructureToReinstantiate())
+	if (UE::StructUtils::FStructReinstancer* Reinstancer = UE::StructUtils::FStructReinstancer::GetInstance())
 	{
-		const UPropertyBag* Bag = GetPropertyBagStruct();
-		if (Bag && Bag->ContainsUserDefinedStruct(StructureToReinstance))
+		if (UPropertyBag* BagStruct = Cast<UPropertyBag>(const_cast<UScriptStruct*>(Value.GetScriptStruct())))
 		{
-			if (StructureToReinstance->Status == EUserDefinedStructureStatus::UDSS_Duplicate)
-			{
-				// On the first pass we create a new bag that contains copy of UDS that represents the currently allocated struct.
-				// GStructureToReinstance is the duplicated struct, and GStructureToReinstance->PrimaryStruct is the UDS that is being reinstanced.
-
-				TArray<FPropertyBagPropertyDesc> PropertyDescs = Bag->PropertyDescs;
-				for (FPropertyBagPropertyDesc& Desc : PropertyDescs)
+			auto MarkAsRequiresReinstanciation = [Reinstancer](TObjectPtr<const UObject>& TypeObject)
 				{
-					if (Desc.ValueType == EPropertyBagPropertyType::Struct)
-					{
-						if (const UUserDefinedStruct* UserDefinedStruct = Cast<UUserDefinedStruct>(Desc.ValueTypeObject))
-						{
-							if (UserDefinedStruct == StructureToReinstance->PrimaryStruct)
-							{
-								Desc.ValueTypeObject = StructureToReinstance;
-							}
-						}					
-					}
-					if (Desc.KeyType == EPropertyBagPropertyType::Struct)
+					const UScriptStruct* Struct = Cast<const UScriptStruct>(TypeObject);
+					if (Struct && Reinstancer->IsReinstanting(Struct))
 					{
-						if (const UUserDefinedStruct* UserDefinedStruct = Cast<UUserDefinedStruct>(Desc.KeyTypeObject))
+						Reinstancer->MarkAsRequiresReinstanciation();
+						// It is not necessary to change the struct to a duplicated one, but it helps with tracking mistakes and prevents memory stomps.
+						if (const UScriptStruct* Duplicated = Reinstancer->GetDuplicatedReinstantingStruct(Struct))
 						{
-							if (UserDefinedStruct == StructureToReinstance->PrimaryStruct)
-							{
-								Desc.KeyTypeObject = StructureToReinstance;
-							}
+							TypeObject = Duplicated;
 						}
 					}
-				}
-
-				const UPropertyBag* NewBag = UPropertyBag::GetOrCreateFromDescs(PropertyDescs);
-				Value.ReplaceScriptStructInternal(NewBag);
-
-				// Adjust recount manually, since we replaced the struct above.
-				Bag->DecrementRefCount();
-				NewBag->IncrementRefCount();
-			}
-			else
+				};
+			TArrayView<FPropertyBagPropertyDesc> PropertyDescs = BagStruct->PropertyDescs;
+			for (FPropertyBagPropertyDesc& Desc : PropertyDescs)
 			{
-				// On the second pass we reinstantiate the data using serialization.
-				// When saving, the UDSs are written using the duplicate which represents current layout, but PrimaryStruct is serialized as the type.
-				// When reading, the data is initialized with the new type, and the serialization will take care of reading from the old data.
-
-				if (UObject* Outer = UE::StructUtils::Private::GetCurrentReinstantiationOuterObject())
+				if (Desc.ValueType == EPropertyBagPropertyType::Struct)
 				{
-					if (!Outer->IsA<UClass>() && !Outer->HasAnyFlags(RF_ClassDefaultObject))
-					{
-						(void)Outer->MarkPackageDirty();
-					}
+					MarkAsRequiresReinstanciation(Desc.ValueTypeObject);
+				}
+				if (Desc.KeyType == EPropertyBagPropertyType::Struct)
+				{
+					MarkAsRequiresReinstanciation(Desc.KeyTypeObject);
 				}
-				
-				TArray<uint8> Data;
-
-				FMemoryWriter Writer(Data);
-				FObjectAndNameAsStringProxyArchive WriterProxy(Writer, /*bInLoadIfFindFails*/true);
-				Serialize(WriterProxy);
-				
-				FMemoryReader Reader(Data);
-				FObjectAndNameAsStringProxyArchive ReaderProxy(Reader, /*bInLoadIfFindFails*/true);
-				Serialize(ReaderProxy);
 			}
 		}
 	}
@@ -4475,7 +4502,7 @@ const FPropertyBagPropertyDesc* UPropertyBag::FindPropertyDescByPropertyName(con
 
 const FPropertyBagPropertyDesc* UPropertyBag::FindPropertyDescByProperty(const FProperty* Property) const
 {
-	return Property ? PropertyDescs.FindByPredicate([&Property](const FPropertyBagPropertyDesc& Desc) { return Desc.CachedProperty && Desc.CachedProperty == Property; }) : nullptr;
+	return Property ? PropertyDescs.FindByPredicate([&Property](const FPropertyBagPropertyDesc& Desc) { return Desc.CachedProperty == Property; }) : nullptr;
 }
 
 const FPropertyBagPropertyDesc* UPropertyBag::FindPropertyDescByIndex(const int32 Index) const
diff --git a/Engine/Source/Runtime/CoreUObject/Private/StructUtils/StructReinstancer.cpp b/Engine/Source/Runtime/CoreUObject/Private/StructUtils/StructReinstancer.cpp
new file mode 100644
--- /dev/null
+++ b/Engine/Source/Runtime/CoreUObject/Private/StructUtils/StructReinstancer.cpp
@@ -0,0 +1,244 @@
+// Copyright Epic Games, Inc. All Rights Reserved.
+
+#include "StructUtils/StructReinstancer.h"
+
+#if WITH_EDITOR
+#include "Logging/StructuredLog.h"
+#include "Serialization/ObjectReader.h"
+#include "Serialization/ObjectWriter.h"
+#include "StructUtils/UserDefinedStruct.h"
+#include "UObject/UObjectIterator.h"
+
+namespace UE::StructUtils::Private
+{
+	static thread_local FStructReinstancer* CurrentReinstancer = nullptr;
+}
+
+namespace UE::StructUtils
+{
+
+FStructReinstancer* FStructReinstancer::GetInstance()
+{
+	return Private::CurrentReinstancer;
+}
+
+void FStructReinstancer::AddStruct(TNotNull<const UScriptStruct*> Original, TNotNull<const UScriptStruct*> Duplicated)
+{
+	StructuresToReinstantiate.Emplace(Original, Duplicated, nullptr);
+}
+
+void FStructReinstancer::SetCompiledStruct(TNotNull<const UScriptStruct*> Original, TNotNull<const UScriptStruct*> NewStruct)
+{
+	FStruct* Found = StructuresToReinstantiate.FindByPredicate(
+		[Original](const FStruct& Other)
+		{
+			return Other.Original == Original;
+		});
+	if (ensure(Found))
+	{
+		Found->New = NewStruct;
+	}
+}
+
+void FStructReinstancer::CollectObjects()
+{
+	if (StructuresToReinstantiate.Num() == 0)
+	{
+		return;
+	}
+
+	// Call AddStructReferencedObjects on all object to find the InstancedStruct/PropertyBag/... that uses the struct.
+
+	// Helper preference collector, does not collect anything, but makes sure AddStructReferencedObjects() gets called e.g. on instanced struct. 
+	class FVisitorReferenceCollector : public FReferenceCollector
+	{
+	public:
+		virtual bool IsIgnoringArchetypeRef() const override
+		{
+			return false;
+		}
+		virtual bool IsIgnoringTransient() const override
+		{
+			return false;
+		}
+		virtual void HandleObjectReference(UObject*& Object, const UObject* ReferencingObject, const FProperty* ReferencingProperty) override
+		{
+			// Empty
+		}
+	};
+	struct FDataForReinstanciation
+	{
+		TNotNull<UObject*> Object;
+		UObject* ClassGeneratedBy = nullptr;
+		TNotNull<UStruct*> Struct;
+		TNotNull<void*> Data;
+		const UStruct* SuperStruct = nullptr;
+		const void* SuperData = nullptr;
+	};
+
+	FVisitorReferenceCollector Collector;
+
+	// bRequiresReinstanciation will be set (in AddStructReferencedObjects) when the struct is reached.
+	check(bSerializingForReinstanciation == false);
+	bRequiresReinstanciation = false;
+
+	TArray<FDataForReinstanciation> ReinstanciationDatas;
+
+	// Handle regular UObjects
+	for (TObjectIterator<UObject> ObjectIt; ObjectIt; ++ObjectIt)
+	{
+		UObject* Object = *ObjectIt;
+
+		Object->CallAddReferencedObjects(Collector);
+		Collector.AddPropertyReferencesWithStructARO(Object->GetClass(), Object);
+
+		if (bRequiresReinstanciation)
+		{
+			ReinstanciationDatas.Emplace(Object, Object->GetClass()->ClassGeneratedBy.Get(), Object->GetClass(), (void*)Object, Object->GetClass(), (const void*)Object->GetClass()->GetDefaultObject());
+			bRequiresReinstanciation = false;
+		}
+	}
+
+	// Handle CDOs and sparse class data 
+	for (TObjectIterator<UClass> ClassIt; ClassIt; ++ClassIt)
+	{
+		UClass* Class = *ClassIt;
+
+		// Handle Sparse class data
+		if (void* SparseData = const_cast<void*>(Class->GetSparseClassData(EGetSparseClassDataMethod::ReturnIfNull)))
+		{
+			UScriptStruct* SparseDataStruct = Class->GetSparseClassDataStruct();
+			Collector.AddPropertyReferencesWithStructARO(SparseDataStruct, SparseData);
+
+			if (bRequiresReinstanciation)
+			{
+				UClass* SuperClass = Class->GetSuperClass();
+				const UScriptStruct* SuperSCDType = SuperClass ? SuperClass->GetSparseClassDataStruct() : nullptr;
+				const void* SuperSCD = SuperClass ? SuperClass->GetSparseClassData(EGetSparseClassDataMethod::ReturnIfNull) : nullptr;
+				ReinstanciationDatas.Emplace(Class, Class->ClassGeneratedBy.Get(), SparseDataStruct, SparseData, SuperSCDType, (const void*)SuperSCD);
+				bRequiresReinstanciation = false;
+			}
+		}
+
+		// Handle CDO
+		if (UObject* CDO = Class->GetDefaultObject())
+		{
+			Collector.AddPropertyReferencesWithStructARO(Class, CDO);
+
+			if (bRequiresReinstanciation)
+			{
+				const UClass* SuperClass = Class->GetSuperClass();
+				const void* SuperData = SuperClass ? SuperClass->GetDefaultObject() : nullptr;
+				ReinstanciationDatas.Emplace(CDO, Class->ClassGeneratedBy.Get(), Class, (void*)CDO, SuperClass, (const void*)SuperData);
+				bRequiresReinstanciation = false;
+			}
+		}
+	}
+
+	ensure(bRequiresReinstanciation == false);
+
+	ObjectsToReinstanciate.Reset();
+	ObjectsToReinstanciate.Reserve(ReinstanciationDatas.Num());
+	for (FDataForReinstanciation& ReinstanciationData : ReinstanciationDatas)
+	{
+		bSerializingForReinstanciation = true;
+
+		UE_LOGFMT(LogClass, Verbose, "Serialize object {Name}", ReinstanciationData.Object->GetFullName());
+
+		TArray<uint8> Buffer;
+		FObjectWriter Writer(Buffer);
+		ReinstanciationData.Struct->SerializeTaggedProperties(
+			Writer,
+			(uint8*)((void*)ReinstanciationData.Data),
+			ReinstanciationData.SuperStruct,
+			(uint8*)ReinstanciationData.SuperData,
+			ReinstanciationData.Object
+		);
+
+		bSerializingForReinstanciation = false;
+
+		ObjectsToReinstanciate.Emplace(
+			FObjectKey(ReinstanciationData.Object),
+			FObjectKey(ReinstanciationData.ClassGeneratedBy),
+			TObjectKey<UStruct>(ReinstanciationData.Struct),
+			ReinstanciationData.Data,
+			TObjectKey<const UStruct>(ReinstanciationData.SuperStruct),
+			ReinstanciationData.SuperData,
+			MoveTemp(Buffer)
+		);
+	}
+}
+
+void FStructReinstancer::ReinstanceObjects()
+{
+	// We want archetype to be re-instantiated before a child in case the child uses the archetype in the serializer.
+	struct FSortedSavedData
+	{
+		UObject* Object;
+		UObject* Archetype;
+		int32 SavedDataIndex;
+	};
+
+	TArray<FSortedSavedData> SortedDatas;
+	SortedDatas.Reserve(ObjectsToReinstanciate.Num());
+	for (int32 Index = 0; Index < ObjectsToReinstanciate.Num(); ++Index)
+	{
+		if (UObject* SavedObject = ObjectsToReinstanciate[Index].Object.ResolveObjectPtr())
+		{
+			SortedDatas.Emplace(SavedObject, SavedObject->GetArchetype(), Index);
+		}
+	}
+
+	SortedDatas.Sort([](const FSortedSavedData& A, const FSortedSavedData& B)
+		{
+			return A.Archetype == nullptr || A.Object == B.Archetype;
+		});
+	for (FSortedSavedData& SortedData : SortedDatas)
+	{
+		const FSavedDataForReinstanciation& SavedData = ObjectsToReinstanciate[SortedData.SavedDataIndex];
+		UObject* SavedObject = SortedData.Object;
+		UStruct* SavedStruct = SavedData.Struct.ResolveObjectPtr();
+		if (ensure(SavedStruct))
+		{
+			if (const UClass* Class = Cast<const UClass>(SavedStruct); Class && Class->HasAnyClassFlags(CLASS_NewerVersionExists))
+			{
+				continue;
+			}
+			if (const UScriptStruct* Struct = Cast<const UScriptStruct>(SavedStruct); Struct && (Struct->StructFlags & STRUCT_NewerVersionExists) != 0)
+			{
+				continue;
+			}
+
+			const UStruct* SavedSuperStruct = SavedData.SuperStruct.ResolveObjectPtr();
+			const void* SavedSuperObjectData = SavedSuperStruct ? SavedData.SuperObjectData : nullptr;
+
+			bSerializingForReinstanciation = true;
+
+			UE_LOGFMT(LogClass, Verbose, "Reinstance object {Name}", SavedObject->GetFullName());
+
+			FObjectReader Reader(SavedData.Buffer);
+			SavedStruct->SerializeTaggedProperties(
+				Reader,
+				(uint8*)(void*)SavedData.ObjectData,
+				SavedSuperStruct,
+				(const uint8*)SavedSuperObjectData);
+
+			bSerializingForReinstanciation = false;
+		}
+	}
+}
+
+FStructReinstancerScope::FStructReinstancerScope()
+{
+	OldStructReinstancer = Private::CurrentReinstancer;
+	Private::CurrentReinstancer = new FStructReinstancer();
+}
+
+FStructReinstancerScope::~FStructReinstancerScope()
+{
+	delete Private::CurrentReinstancer;
+	Private::CurrentReinstancer = OldStructReinstancer;
+}
+
+}
+#endif // WITH_EDITOR
diff --git a/Engine/Source/Runtime/CoreUObject/Private/StructUtils/StructUtilsTypes.cpp b/Engine/Source/Runtime/CoreUObject/Private/StructUtils/StructUtilsTypes.cpp
--- a/Engine/Source/Runtime/CoreUObject/Private/StructUtils/StructUtilsTypes.cpp
+++ b/Engine/Source/Runtime/CoreUObject/Private/StructUtils/StructUtilsTypes.cpp
@@ -234,44 +234,4 @@ namespace UE::StructUtils
 	{
 		return GetStructHash64Helper(SharedView);
 	}
-
-	namespace Private
-	{
-#if WITH_EDITOR
-		const UUserDefinedStruct* GStructureToReinstantiate = nullptr;
-		UObject* GCurrentReinstantiationOuterObject = nullptr;
-
-		FStructureToReinstantiateScope::FStructureToReinstantiateScope(const UUserDefinedStruct* StructureToReinstantiate)
-		{
-			OldStructureToReinstantiate = GStructureToReinstantiate;
-			GStructureToReinstantiate = StructureToReinstantiate;
-		}
-	
-		FStructureToReinstantiateScope::~FStructureToReinstantiateScope()
-		{
-			GStructureToReinstantiate = OldStructureToReinstantiate;
-		}
-
-		FCurrentReinstantiationOuterObjectScope::FCurrentReinstantiationOuterObjectScope(UObject* CurrentReinstantiateOuterObject)
-		{
-			OldCurrentReinstantiateOuterObject = GCurrentReinstantiationOuterObject;
-			GCurrentReinstantiationOuterObject = CurrentReinstantiateOuterObject;
-		}
-
-		FCurrentReinstantiationOuterObjectScope::~FCurrentReinstantiationOuterObjectScope()
-		{
-			GCurrentReinstantiationOuterObject = OldCurrentReinstantiateOuterObject;
-		}
-
-		const UUserDefinedStruct* GetStructureToReinstantiate()
-		{
-			return GStructureToReinstantiate;
-		}
-	
-		UObject* GetCurrentReinstantiationOuterObject()
-		{
-			return GCurrentReinstantiationOuterObject;
-		}
-#endif // WITH_EDITOR
-	}
 }
diff --git a/Engine/Source/Runtime/CoreUObject/Public/StructUtils/StructReinstancer.h b/Engine/Source/Runtime/CoreUObject/Public/StructUtils/StructReinstancer.h
new file mode 100644
--- /dev/null
+++ b/Engine/Source/Runtime/CoreUObject/Public/StructUtils/StructReinstancer.h
@@ -0,0 +1,133 @@
+// Copyright Epic Games, Inc. All Rights Reserved.
+
+#pragma once
+
+#include "UObject/Class.h"
+#include "UObject/ObjectKey.h"
+
+#define UE_API COREUOBJECT_API
+
+#if WITH_EDITOR
+namespace UE::StructUtils
+{
+	/**
+	 * Re-instance UScriptStruct with their new version.
+	 * It uses the Reference Collector (AddStructReferencedObjects) to find all InstancedStruct/PropertyBag/... that use the struct.
+	 * To restore only the modified property, the UObject that owns the UScriptStruct is serialized before the compilation process starts, and then serialized back once the new UScriptStruct is compiled.
+	 * It utilizes tagged serialization and the default value of the archetype to obtain the new default values and restore the previously modified values.
+	 * The tagged serialization only serialized the modified properties. For proper reinitialization, the owning struct must be different from the archetype. Override Identical and return false to make sure the owning struct is serialized.
+	 * For UUserDefinedStruct re-instantiation, the original UScriptStruct and the new compiled UScriptStruct are the same instance. Other UScriptStruct types use different instances.
+	 * 
+	 * Algo:
+	 * 1. For UUserDefinedStruct, create a duplicate. It will be used instead of the original UScriptStruct until it is compiled.
+	 * 2. Gather all UObjects that contain a soon-to-be-dirty UScriptStruct.
+	 *   1. For UUserDefinedStruct, replace the UScriptStruct with the duplicate.
+	 * 3. Serialize (save) the UObjects (to gather only the modified properties).
+	 * 4. Compile the UScriptStruct.
+	 * 5. Serialize (load) the same UObjects. Replace the UScriptStrut with the new compiled version.
+	 */
+	class UE_INTERNAL FStructReinstancer final : FNoncopyable
+	{
+	public:
+		UE_API static FStructReinstancer* GetInstance();
+		/**
+		 * Add a UScriptStrut to the re-instancer and its duplicated struct.
+		 * The duplicated UScriptStrut is used to keep the raw memory readable while the UScriptStrut is compiled.
+		 */
+		UE_API void AddStruct(TNotNull<const UScriptStruct*> Original, TNotNull<const UScriptStruct*> Duplicated);
+
+		/** Set the new compiled UScriptStrut created for the original struct. */
+		UE_API void SetCompiledStruct(TNotNull<const UScriptStruct*> Original, TNotNull<const UScriptStruct*> Compiled);
+
+		/** Collect all objects that use structs that will be re-instantiated. */
+		UE_API void CollectObjects();
+
+		/** Re-instantiate previously collected objects with the new struct. */
+		UE_API void ReinstanceObjects();
+
+		/** Serializing the UObject owners to gather the modified properties. */
+		bool IsSerializingForReinstanciation() const
+		{
+			return bSerializingForReinstanciation;
+		}
+
+		/** Is the struct added to the FStructReinstancer. */
+		bool IsReinstanting(TNotNull<const UScriptStruct*> Struct) const
+		{
+			return StructuresToReinstantiate.ContainsByPredicate(
+				[Struct](const FStruct& Other)
+				{
+					return Other.Original.Get() == Struct
+						|| Other.Duplicated.Get() == Struct;
+				});
+		}
+
+		/**
+		 * Find the duplicated UScriptStruct created to keep the raw memory readable while the struct is compiled.
+		 * @note It is different for UUserDefinedStruct.
+		 */
+		const UScriptStruct* GetDuplicatedReinstantingStruct(TNotNull<const UScriptStruct*> Struct) const
+		{
+			const FStruct* Found = StructuresToReinstantiate.FindByPredicate(
+				[Struct](const FStruct& Other)
+				{
+					return Other.Original.Get() == Struct;
+				});
+			return Found ? Found->Duplicated.Get() : nullptr;
+		}
+
+		/**
+		 * Find the new compiled UScriptStruct created for the original struct.
+		 * @note The new compiled UScriptStruct is the same as the original for UUserDefinedStruct.
+		 */
+		const UScriptStruct* GetCompiledReinstantingStruct(TNotNull<const UScriptStruct*> Struct) const
+		{
+			const FStruct* Found = StructuresToReinstantiate.FindByPredicate(
+				[Struct](const FStruct& Other)
+				{
+					return Other.Original == Struct
+						|| Other.Duplicated == Struct;
+				});
+			return Found ? Found->New.Get() : nullptr;
+		}
+
+		/** The current collected UObject contains at least one UScriptStruct that is being reinstanciated. */
+		void MarkAsRequiresReinstanciation()
+		{
+			bRequiresReinstanciation = true;
+		}
+	private:
+		struct FStruct
+		{
+			TObjectPtr<const UScriptStruct> Original;
+			TObjectPtr<const UScriptStruct> Duplicated;
+			TObjectPtr<const UScriptStruct> New;
+		};
+		struct FSavedDataForReinstanciation
+		{
+			FObjectKey Object;
+			FObjectKey ClassGeneratedBy;
+			TObjectKey<UStruct> Struct;
+			TNotNull<void*> ObjectData;
+			TObjectKey<const UStruct> SuperStruct;
+			const void* SuperObjectData = nullptr;
+			TArray<uint8> Buffer;
+		};
+		TArray<FStruct> StructuresToReinstantiate;
+		TArray<FSavedDataForReinstanciation> ObjectsToReinstanciate;
+		bool bRequiresReinstanciation = false;
+		bool bSerializingForReinstanciation = false;
+	};
+
+	struct UE_INTERNAL FStructReinstancerScope final : FNoncopyable
+	{
+		UE_API explicit FStructReinstancerScope();
+		UE_API ~FStructReinstancerScope();
+	private:
+		FStructReinstancer* OldStructReinstancer = nullptr;
+	};
+}
+
+#endif // WITH_EDITOR
+
+#undef UE_API
diff --git a/Engine/Source/Runtime/CoreUObject/Public/StructUtils/StructUtilsTypes.h b/Engine/Source/Runtime/CoreUObject/Public/StructUtils/StructUtilsTypes.h
--- a/Engine/Source/Runtime/CoreUObject/Public/StructUtils/StructUtilsTypes.h
+++ b/Engine/Source/Runtime/CoreUObject/Public/StructUtils/StructUtilsTypes.h
@@ -129,24 +129,27 @@ struct FStructTypeSortOperator
 namespace UE::StructUtils::Private
 {
 	// Private structs used during user defined struct reinstancing.
-	struct FStructureToReinstantiateScope
+	struct UE_DEPRECATED(5.8, "Use FStructReinstancer instead.") FStructureToReinstantiateScope
 	{
-		UE_API explicit FStructureToReinstantiateScope(const UUserDefinedStruct* StructureToReinstantiate);
-		UE_API ~FStructureToReinstantiateScope();
-	private:
-		const UUserDefinedStruct* OldStructureToReinstantiate = nullptr;
+		explicit FStructureToReinstantiateScope(const UUserDefinedStruct* StructureToReinstantiate) {}
 	};
 
-	struct FCurrentReinstantiationOuterObjectScope
+	struct UE_DEPRECATED(5.8, "Use FStructReinstancer instead.") FCurrentReinstantiationOuterObjectScope
 	{
-		UE_API explicit FCurrentReinstantiationOuterObjectScope(UObject* CurrentReinstantiateOuterObject);
-		UE_API ~FCurrentReinstantiationOuterObjectScope();
-	private:
-		UObject* OldCurrentReinstantiateOuterObject = nullptr;
+		explicit FCurrentReinstantiationOuterObjectScope(UObject* CurrentReinstantiateOuterObject) {}
 	};
 
-	extern UE_API const UUserDefinedStruct* GetStructureToReinstantiate();
-	extern UE_API UObject* GetCurrentReinstantiationOuterObject();
+	UE_DEPRECATED(5.8, "Use FStructReinstancer instead.")
+	static const UUserDefinedStruct* GetStructureToReinstantiate()
+	{
+		return nullptr;
+	}
+
+	UE_DEPRECATED(5.8, "Use FStructReinstancer instead.")
+	static UObject* GetCurrentReinstantiationOuterObject()
+	{
+		return nullptr;
+	}
 };
 #endif // WITH_EDITOR
 
