Fix build issue 1191905. The Vect property was not properly initialized.
#rnx

#Synced-CL 49772122

diff --git a/Engine/Source/Editor/StructUtilsEditor/Private/Tests/StructReinstancerTests.h b/Engine/Source/Editor/StructUtilsEditor/Private/Tests/StructReinstancerTests.h
--- a/Engine/Source/Editor/StructUtilsEditor/Private/Tests/StructReinstancerTests.h
+++ b/Engine/Source/Editor/StructUtilsEditor/Private/Tests/StructReinstancerTests.h
@@ -18,7 +18,7 @@ struct FTestStructThatContainsStruct
 	GENERATED_BODY()
 
 	UPROPERTY()
-	FVector Vect;
+	FVector Vect = FVector::ZeroVector;
 
 	UPROPERTY()
 	FInstancedStruct InstancedStruct;
@@ -37,7 +37,7 @@ class UTestObjectThatContainsStructBase : public UObject
 
 public:
 	UPROPERTY()
-	FVector Vect;
+	FVector Vect = FVector::ZeroVector;
 
 	UPROPERTY()
 	FTestStructThatContainsStruct Container;
