Add a TheSameLayout that supports an empty FInstancedPropertyBag.
Always create an empty UUserDefinedStruct when the FInstancePropertyBag is invalid. The UUserDefinedStruct is used as a reference. Even when empty, we want to maintain that reference's validity. It also behaves the same with a UDS, which used to have properties, and gets updated without properties.

#Synced-CL 49833236

#Synced-CL 49833236

diff --git a/Engine/Source/Editor/StructUtilsEditor/Private/StructUtilsEditorUtils.cpp b/Engine/Source/Editor/StructUtilsEditor/Private/StructUtilsEditorUtils.cpp
--- a/Engine/Source/Editor/StructUtilsEditor/Private/StructUtilsEditorUtils.cpp
+++ b/Engine/Source/Editor/StructUtilsEditor/Private/StructUtilsEditorUtils.cpp
@@ -17,7 +17,8 @@ namespace UE::StructUtils::Private
 		UUserDefinedStruct* Struct = NewObject<UUserDefinedStruct>(InOuter, Name, RF_Transactional | RF_Public);
 		check(Struct);
 
-		Struct->EditorData = NewObject<UUserDefinedStructEditorData>(Struct, FName(), RF_Transactional);
+		UUserDefinedStructEditorData* UDSEditorData = NewObject<UUserDefinedStructEditorData>(Struct, FName(), RF_Transactional);
+		Struct->EditorData = UDSEditorData;
 		check(Struct->EditorData);
 
 		Struct->Guid = FGuid::NewGuid();
@@ -25,8 +26,9 @@ namespace UE::StructUtils::Private
 		Struct->Bind();
 		constexpr bool bRelinkExistingProperties = true;
 		Struct->StaticLink(bRelinkExistingProperties);
+		UDSEditorData->RecreateDefaultInstance();
 
-		Struct->Status = UDSS_Error;
+		Struct->Status = UDSS_UpToDate;
 		return Struct;
 	}
 
@@ -86,23 +88,23 @@ namespace UE::StructUtils
 {
 	TValueOrError<UUserDefinedStruct*, void> CreateUserDefinedStructFromDescs(TNotNull<UObject*> UDSPackage, const FInstancedPropertyBag& PropertyBag, FName UserDefinedStructName)
 	{
-		if (!PropertyBag.IsValid())
+		if (Private::FindUserDefinedStruct(UDSPackage, UserDefinedStructName))
 		{
-			return MakeValue(nullptr);
+			return MakeError();
 		}
 
-		if (Private::FindUserDefinedStruct(UDSPackage, UserDefinedStructName))
+		TNotNull<UUserDefinedStruct*> UserDefinedStruct = Private::CreateUserDefinedStruct(UDSPackage, UserDefinedStructName);
+		if (!PropertyBag.IsValid())
 		{
-			return MakeError();
+			return MakeValue(UserDefinedStruct);
 		}
 
 		TArrayView<const FPropertyBagPropertyDesc> PropertyDescs = PropertyBag.GetPropertyBagStruct()->GetPropertyDescs();
 		if (PropertyDescs.IsEmpty())
 		{
-			return MakeValue(nullptr);
+			return MakeValue(UserDefinedStruct);
 		}
 
-		TNotNull<UUserDefinedStruct*> UserDefinedStruct = Private::CreateUserDefinedStruct(UDSPackage, UserDefinedStructName);
 		TNotNull<UUserDefinedStructEditorData*> EditorData = CastChecked<UUserDefinedStructEditorData>(UserDefinedStruct->EditorData);
 		EditorData->VariablesDescriptions = Private::GetPropertyDescs(PropertyBag, PropertyDescs);
 		FStructureEditorUtils::CompileStructure(UserDefinedStruct);
diff --git a/Engine/Source/Runtime/CoreUObject/Private/StructUtils/StructUtilsTypes.cpp b/Engine/Source/Runtime/CoreUObject/Private/StructUtils/StructUtilsTypes.cpp
--- a/Engine/Source/Runtime/CoreUObject/Private/StructUtils/StructUtilsTypes.cpp
+++ b/Engine/Source/Runtime/CoreUObject/Private/StructUtils/StructUtilsTypes.cpp
@@ -3,6 +3,7 @@
 #include "StructUtils/StructUtilsTypes.h"
 #include "Hash/CityHash.h"
 #include "Serialization/ArchiveCrc32.h"
+#include "StructUtils/PropertyBag.h"
 #include "StructUtils/StructView.h"
 #include "StructUtils/SharedStruct.h"
 #include "StructUtils/UserDefinedStruct.h"
@@ -234,4 +235,15 @@ namespace UE::StructUtils
 	{
 		return GetStructHash64Helper(SharedView);
 	}
+
+	bool TheSameLayout(const FInstancedPropertyBag& A, const UScriptStruct* B)
+	{
+		// Both invalid or Both empty.
+		if (!A.IsValid())
+		{
+			return B == nullptr || B->PropertyLink == nullptr;
+		}
+
+		return FStructUtils::TheSameLayout(A.GetPropertyBagStruct(), B);
+	}
 }
diff --git a/Engine/Source/Runtime/CoreUObject/Public/StructUtils/StructUtilsTypes.h b/Engine/Source/Runtime/CoreUObject/Public/StructUtils/StructUtilsTypes.h
--- a/Engine/Source/Runtime/CoreUObject/Public/StructUtils/StructUtilsTypes.h
+++ b/Engine/Source/Runtime/CoreUObject/Public/StructUtils/StructUtilsTypes.h
@@ -12,6 +12,7 @@ struct FSharedStruct;
 struct FConstSharedStruct;
 struct FStructView;
 struct FConstStructView;
+struct FInstancedPropertyBag;
 struct FInstancedStruct;
 class FReferenceCollector;
 class UUserDefinedStruct;
@@ -68,6 +69,13 @@ namespace UE::StructUtils
 		|| std::is_same_v<FSharedStruct, T>
 		|| std::is_same_v<FConstSharedStruct, T>
 		|| std::is_same_v<FInstancedStruct, T>;
+
+	/**
+	 * Whether an FInstancedPropertyBag and an UScriptStruct have exactly the same memory layout.
+	 * The struct UPropertyBag might not be generated if there are no properties in the struct.
+	 * In that case, it matches if the UScriptStruct also doesn’t have properties.
+	 */
+	extern UE_API bool TheSameLayout(const FInstancedPropertyBag& A, const UScriptStruct* B);
 }
 
 /* Predicate useful to find a struct of a specific type in an container */
