[StructUtils] Add metadata HideStructType to InstancedStruct's detail customizer. It prevents the user from seeing and changing the struct type.
Delay the DetailView rebuild after a UserDefinedStruct compiles. The compile can run while loading an asset (on the game thread), but the detail ensures that it detects that assets are loading.

#Synced-CL 49855772

diff --git a/Engine/Source/Editor/StructUtilsEditor/Private/InstancedStructDetails.cpp b/Engine/Source/Editor/StructUtilsEditor/Private/InstancedStructDetails.cpp
--- a/Engine/Source/Editor/StructUtilsEditor/Private/InstancedStructDetails.cpp
+++ b/Engine/Source/Editor/StructUtilsEditor/Private/InstancedStructDetails.cpp
@@ -18,6 +18,7 @@
 #include "IStructureDataProvider.h"
 #include "Misc/ConfigCacheIni.h"
 #include "StructUtilsDelegates.h"
+#include "StructUtilsMetadata.h"
 #include "SInstancedStructPicker.h"
 
 #define LOCTEXT_NAMESPACE "StructUtilsEditor"
@@ -283,7 +284,7 @@ void FInstancedStructDataDetails::GetPropertyGroups(
 void FInstancedStructDataDetails::OnStructLayoutChanges()
 {
 	bCanHandleStructValuePostChange = false;
-	OnRegenerateChildren.ExecuteIfBound();
+	bStructChangeDetected = true;
 }
 
 void FInstancedStructDataDetails::OnStructHandlePostChange()
@@ -293,6 +294,7 @@ void FInstancedStructDataDetails::OnStructHandlePostChange()
 		TArray<TWeakObjectPtr<const UStruct>> InstanceTypes = GetInstanceTypes();
 		if (InstanceTypes != CachedInstanceTypes)
 		{
+			bStructChangeDetected = false;
 			OnRegenerateChildren.ExecuteIfBound();
 		}
 	}
@@ -384,8 +386,9 @@ void FInstancedStructDataDetails::Tick(float DeltaTime)
 {
 	// If the instance types change (e.g. due to selecting new struct type), we'll need to update the layout.
 	TArray<TWeakObjectPtr<const UStruct>> InstanceTypes = GetInstanceTypes();
-	if (InstanceTypes != CachedInstanceTypes)
+	if (bStructChangeDetected || InstanceTypes != CachedInstanceTypes)
 	{
+		bStructChangeDetected = false;
 		OnRegenerateChildren.ExecuteIfBound();
 	}
 }
@@ -419,20 +422,23 @@ void FInstancedStructDetails::CustomizeHeader(TSharedRef<class IPropertyHandle>
 
 	OnObjectsReinstancedHandle = FCoreUObjectDelegates::OnObjectsReinstanced.AddSP(this, &FInstancedStructDetails::OnObjectsReinstanced);
 
-	HeaderRow
-		.CopyAction(FUIAction(FExecuteAction::CreateSP(this, &FInstancedStructDetails::OnCopyInstancedStruct)))
-		.PasteAction(FUIAction(FExecuteAction::CreateSP(this, &FInstancedStructDetails::OnPasteInstancedStruct), FCanExecuteAction::CreateSP(this, &FInstancedStructDetails::CanPasteInstancedStruct)))
-		.NameContent()
-		[
-			StructPropertyHandle->CreatePropertyNameWidget()
-		]
-		.ValueContent()
-		.MinDesiredWidth(250.f)
-		.VAlign(VAlign_Center)
-		[
-			SAssignNew(StructPicker, SInstancedStructPicker, StructProperty, PropUtils)
-		]
-		.IsEnabled(StructProperty->IsEditable());
+	if (!StructProperty->HasMetaData(UE::StructUtils::Metadata::HideStructTypeName) || StructProperty->GetMetaData(UE::StructUtils::Metadata::HideStructTypeName) == TEXT("false"))
+	{
+		HeaderRow
+			.CopyAction(FUIAction(FExecuteAction::CreateSP(this, &FInstancedStructDetails::OnCopyInstancedStruct)))
+			.PasteAction(FUIAction(FExecuteAction::CreateSP(this, &FInstancedStructDetails::OnPasteInstancedStruct), FCanExecuteAction::CreateSP(this, &FInstancedStructDetails::CanPasteInstancedStruct)))
+			.NameContent()
+			[
+				StructPropertyHandle->CreatePropertyNameWidget()
+			]
+			.ValueContent()
+			.MinDesiredWidth(250.f)
+			.VAlign(VAlign_Center)
+			[
+				SAssignNew(StructPicker, SInstancedStructPicker, StructProperty, PropUtils)
+			]
+			.IsEnabled(StructProperty->IsEditable());
+	}
 }
 
 void FInstancedStructDetails::OnCopyInstancedStruct() const
diff --git a/Engine/Source/Editor/StructUtilsEditor/Private/StructUtilsMetadata.cpp b/Engine/Source/Editor/StructUtilsEditor/Private/StructUtilsMetadata.cpp
--- a/Engine/Source/Editor/StructUtilsEditor/Private/StructUtilsMetadata.cpp
+++ b/Engine/Source/Editor/StructUtilsEditor/Private/StructUtilsMetadata.cpp
@@ -27,6 +27,7 @@ namespace UE::StructUtils::Metadata
 	const FLazyName BaseStructName = "BaseStruct";
 	const FLazyName BaseStructFunctionName = "BaseStructFunction";
 	const FLazyName StructTypeConstName = "StructTypeConst";
+	const FLazyName HideStructTypeName = "HideStructType";
 
 	bool AreCategoriesEnabled(const FPropertyBagPropertyDesc& Desc)
 	{
diff --git a/Engine/Source/Editor/StructUtilsEditor/Public/InstancedStructDetails.h b/Engine/Source/Editor/StructUtilsEditor/Public/InstancedStructDetails.h
--- a/Engine/Source/Editor/StructUtilsEditor/Public/InstancedStructDetails.h
+++ b/Engine/Source/Editor/StructUtilsEditor/Public/InstancedStructDetails.h
@@ -118,6 +118,8 @@ private:
 
 	/** True if we're allowed to handle a StructValuePostChange */
 	bool bCanHandleStructValuePostChange = false;
+	/** Mark the struct has changed, it requires a rebuild on the next iteration. */
+	bool bStructChangeDetected = false;
 	
 	FDelegateHandle UserDefinedStructReinstancedHandle;
 
diff --git a/Engine/Source/Editor/StructUtilsEditor/Public/StructUtilsMetadata.h b/Engine/Source/Editor/StructUtilsEditor/Public/StructUtilsMetadata.h
--- a/Engine/Source/Editor/StructUtilsEditor/Public/StructUtilsMetadata.h
+++ b/Engine/Source/Editor/StructUtilsEditor/Public/StructUtilsMetadata.h
@@ -61,6 +61,9 @@ namespace UE::StructUtils::Metadata
 	/// [PropertyMetadata] StructTypeConst: Determine whether the type of the instanced struct can be changed
 	extern STRUCTUTILSEDITOR_API const FLazyName StructTypeConstName;
 
+	/// [PropertyMetadata] HideStructType: Determine whether the type of the instanced struct is visibled
+	extern STRUCTUTILSEDITOR_API const FLazyName HideStructTypeName;
+
 	/** If organizing by category is enabled for this property. */
 	STRUCTUTILSEDITOR_API bool AreCategoriesEnabled(const FPropertyBagPropertyDesc& Desc);
 	/** Enable category organization for this property. */
