[StructUtils] Fixes for UE::StructUtils::CreateUserDefinedStructFromDescs.
Add additional options for constructing the UDS.
The PropertyBag generates a FEnumProperty. Ensure the UDS also generates the same property type.

#Synced-CL 49855772

diff --git a/Engine/Source/Editor/StructUtilsEditor/Private/Customizations/PropertyBagDetails.cpp b/Engine/Source/Editor/StructUtilsEditor/Private/Customizations/PropertyBagDetails.cpp
--- a/Engine/Source/Editor/StructUtilsEditor/Private/Customizations/PropertyBagDetails.cpp
+++ b/Engine/Source/Editor/StructUtilsEditor/Private/Customizations/PropertyBagDetails.cpp
@@ -247,8 +247,7 @@ FEdGraphPinType GetPropertyDescAsPin(const FPropertyBagPropertyDesc& Desc)
 			OutCategory = UEdGraphSchema_K2::PC_Text;
 			break;
 		case EPropertyBagPropertyType::Enum:
-			// @todo: some pin coloring is not correct due to this (byte-as-enum vs enum). 
-			OutCategory = UEdGraphSchema_K2::PC_Enum;
+			OutCategory = UEdGraphSchema_K2::PC_Byte;
 			OutSubCategoryObject = const_cast<UObject*>(TypeObject.Get());
 			break;
 		case EPropertyBagPropertyType::Struct:
diff --git a/Engine/Source/Editor/StructUtilsEditor/Private/StructUtilsEditorUtils.cpp b/Engine/Source/Editor/StructUtilsEditor/Private/StructUtilsEditorUtils.cpp
--- a/Engine/Source/Editor/StructUtilsEditor/Private/StructUtilsEditorUtils.cpp
+++ b/Engine/Source/Editor/StructUtilsEditor/Private/StructUtilsEditorUtils.cpp
@@ -12,9 +12,11 @@
 
 namespace UE::StructUtils::Private
 {
-	TNotNull<UUserDefinedStruct*> CreateUserDefinedStruct(TNotNull<UObject*> InOuter, FName Name)
+	const FName NAME_GeneratePropertyWithoutID = "GeneratePropertyWithoutID";
+
+	TNotNull<UUserDefinedStruct*> CreateUserDefinedStruct(TNotNull<UObject*> InOuter, FName Name, const FCreateUserDefinedStructArgs& Args)
 	{
-		UUserDefinedStruct* Struct = NewObject<UUserDefinedStruct>(InOuter, Name, RF_Transactional | RF_Public);
+		UUserDefinedStruct* Struct = NewObject<UUserDefinedStruct>(InOuter, Name, Args.UserDefinedStructFlags);
 		check(Struct);
 
 		UUserDefinedStructEditorData* UDSEditorData = NewObject<UUserDefinedStructEditorData>(Struct, FName(), RF_Transactional);
@@ -22,7 +24,15 @@ namespace UE::StructUtils::Private
 		check(Struct->EditorData);
 
 		Struct->Guid = FGuid::NewGuid();
-		Struct->SetMetaData(TEXT("BlueprintType"), TEXT("true"));
+		if (Args.bIsBlueprintType)
+		{
+			Struct->SetMetaData("BlueprintType", TEXT("true"));
+		}
+		if (!Args.bAddPropertyIDToPropertyName)
+		{
+			Struct->SetMetaData(NAME_GeneratePropertyWithoutID, TEXT("true"));
+		}
+
 		Struct->Bind();
 		constexpr bool bRelinkExistingProperties = true;
 		Struct->StaticLink(bRelinkExistingProperties);
@@ -45,23 +55,27 @@ namespace UE::StructUtils::Private
 			&& A.DefaultValue == B.DefaultValue;
 	}
 
-	FStructVariableDescription CreateUserDefinedVariableFromDesc(const FPropertyBagPropertyDesc& PropertyDesc)
+	FStructVariableDescription CreateUserDefinedVariableFromDesc(const FPropertyBagPropertyDesc& PropertyDesc, bool bNameWithID)
 	{
 		const FEdGraphPinType GraphPin = UE::StructUtils::GetPropertyDescAsPin(PropertyDesc);
 
 		FStructVariableDescription Result;
 		Result.SetPinType(GraphPin);
 		Result.FriendlyName = PropertyDesc.Name.ToString();
-		// FMemberVariableNameHelper adds the ID to the name.
-		Result.VarName = *FString::Printf(TEXT("%s_%s"), *Result.FriendlyName, *PropertyDesc.ID.ToString(EGuidFormats::Digits));
+		Result.VarName = PropertyDesc.Name;
+		if (bNameWithID)
+		{
+			// FMemberVariableNameHelper adds the ID to the name.
+			Result.VarName = *FString::Printf(TEXT("%s_%s"), *Result.FriendlyName, *PropertyDesc.ID.ToString(EGuidFormats::Digits));
+		}
 		Result.VarGuid = PropertyDesc.ID;
 
 		return Result;
 	}
 
-	FStructVariableDescription CreateUserDefinedVariableFromDesc(const FInstancedPropertyBag& PropertyBag, const FPropertyBagPropertyDesc& PropertyDesc)
+	FStructVariableDescription CreateUserDefinedVariableFromDesc(const FInstancedPropertyBag& PropertyBag, const FPropertyBagPropertyDesc& PropertyDesc, bool bNameWithID)
 	{
-		FStructVariableDescription Result = CreateUserDefinedVariableFromDesc(PropertyDesc);
+		FStructVariableDescription Result = CreateUserDefinedVariableFromDesc(PropertyDesc, bNameWithID);
 		if (const FProperty* Property = PropertyBag.GetValue().GetScriptStruct()->FindPropertyByName(PropertyDesc.Name))
 		{
 			UObject* OwningObject = const_cast<UScriptStruct*>(PropertyBag.GetValue().GetScriptStruct());
@@ -71,14 +85,14 @@ namespace UE::StructUtils::Private
 		return Result;
 	}
 
-	TArray<FStructVariableDescription> GetPropertyDescs(const FInstancedPropertyBag& PropertyBag, const TArrayView<const FPropertyBagPropertyDesc>  PropertyDescs)
+	TArray<FStructVariableDescription> GetPropertyDescs(const FInstancedPropertyBag& PropertyBag, const TArrayView<const FPropertyBagPropertyDesc>  PropertyDescs, bool bNameWithID)
 	{
 		TArray<FStructVariableDescription> VariablesDescriptions;
 		VariablesDescriptions.Reserve(PropertyDescs.Num());
 		TArray<FPropertyBagPropertyDesc> SanitizedPropertyDescs = UPropertyBag::SanitizePropertyDescs(PropertyDescs);
 		for (const FPropertyBagPropertyDesc& Desc : SanitizedPropertyDescs)
 		{
-			VariablesDescriptions.Add(CreateUserDefinedVariableFromDesc(PropertyBag, Desc));
+			VariablesDescriptions.Add(CreateUserDefinedVariableFromDesc(PropertyBag, Desc, bNameWithID));
 		}
 		return VariablesDescriptions;
 	}
@@ -86,14 +100,14 @@ namespace UE::StructUtils::Private
 
 namespace UE::StructUtils
 {
-	TValueOrError<UUserDefinedStruct*, void> CreateUserDefinedStructFromDescs(TNotNull<UObject*> UDSPackage, const FInstancedPropertyBag& PropertyBag, FName UserDefinedStructName)
+	TValueOrError<UUserDefinedStruct*, void> CreateUserDefinedStructFromDescs(TNotNull<UObject*> UDSPackage, const FInstancedPropertyBag& PropertyBag, FName UserDefinedStructName, const FCreateUserDefinedStructArgs& Args)
 	{
 		if (Private::FindUserDefinedStruct(UDSPackage, UserDefinedStructName))
 		{
 			return MakeError();
 		}
 
-		TNotNull<UUserDefinedStruct*> UserDefinedStruct = Private::CreateUserDefinedStruct(UDSPackage, UserDefinedStructName);
+		TNotNull<UUserDefinedStruct*> UserDefinedStruct = Private::CreateUserDefinedStruct(UDSPackage, UserDefinedStructName, Args);
 		if (!PropertyBag.IsValid())
 		{
 			return MakeValue(UserDefinedStruct);
@@ -106,7 +120,7 @@ namespace UE::StructUtils
 		}
 
 		TNotNull<UUserDefinedStructEditorData*> EditorData = CastChecked<UUserDefinedStructEditorData>(UserDefinedStruct->EditorData);
-		EditorData->VariablesDescriptions = Private::GetPropertyDescs(PropertyBag, PropertyDescs);
+		EditorData->VariablesDescriptions = Private::GetPropertyDescs(PropertyBag, PropertyDescs, Args.bAddPropertyIDToPropertyName);
 		FStructureEditorUtils::CompileStructure(UserDefinedStruct);
 
 		return MakeValue(UserDefinedStruct);
@@ -124,7 +138,8 @@ namespace UE::StructUtils
 				PropertyDescs = PropertyBag.GetPropertyBagStruct()->GetPropertyDescs();
 			}
 
-			TArray<FStructVariableDescription> VariablesDescriptions = Private::GetPropertyDescs(PropertyBag, PropertyDescs);
+			const bool bNameWithoutID = UserDefinedStruct->GetMetaData(Private::NAME_GeneratePropertyWithoutID) == TEXT("true");
+			TArray<FStructVariableDescription> VariablesDescriptions = Private::GetPropertyDescs(PropertyBag, PropertyDescs, !bNameWithoutID);
 			bEqual = EditorData->VariablesDescriptions.Num() == VariablesDescriptions.Num();
 			if (bEqual)
 			{
diff --git a/Engine/Source/Editor/StructUtilsEditor/Private/Tests/StructReinstancerTests.cpp b/Engine/Source/Editor/StructUtilsEditor/Private/Tests/StructReinstancerTests.cpp
--- a/Engine/Source/Editor/StructUtilsEditor/Private/Tests/StructReinstancerTests.cpp
+++ b/Engine/Source/Editor/StructUtilsEditor/Private/Tests/StructReinstancerTests.cpp
@@ -109,7 +109,7 @@ namespace UE::StructUtils::Private
 
 		UUserDefinedStruct* UserDefinedStruct = nullptr;
 		{
-			TValueOrError<UUserDefinedStruct*, void> Result = UE::StructUtils::CreateUserDefinedStructFromDescs(UDSContainer.Get(), BagA, "UDS");
+			TValueOrError<UUserDefinedStruct*, void> Result = UE::StructUtils::CreateUserDefinedStructFromDescs(UDSContainer.Get(), BagA, "UDS", UE::StructUtils::FCreateUserDefinedStructArgs());
 
 			UTEST_TRUE("UUserDefinedStruct created", Result.HasValue());
 			UserDefinedStruct = Result.GetValue();
@@ -245,7 +245,7 @@ namespace UE::StructUtils::Private
 
 		UUserDefinedStruct* UserDefinedStruct = nullptr;
 		{
-			TValueOrError<UUserDefinedStruct*, void> Result = UE::StructUtils::CreateUserDefinedStructFromDescs(UDSContainer.Get(), BagA, "UDS");
+			TValueOrError<UUserDefinedStruct*, void> Result = UE::StructUtils::CreateUserDefinedStructFromDescs(UDSContainer.Get(), BagA, "UDS", UE::StructUtils::FCreateUserDefinedStructArgs());
 
 			UTEST_TRUE("UUserDefinedStruct created", Result.HasValue());
 			UserDefinedStruct = Result.GetValue();
@@ -367,7 +367,7 @@ namespace UE::StructUtils::Private
 		// Create UDS
 		UUserDefinedStruct* UserDefinedStruct = nullptr;
 		{
-			TValueOrError<UUserDefinedStruct*, void> Result = UE::StructUtils::CreateUserDefinedStructFromDescs(UDSContainer.Get(), BagA, "UDS");
+			TValueOrError<UUserDefinedStruct*, void> Result = UE::StructUtils::CreateUserDefinedStructFromDescs(UDSContainer.Get(), BagA, "UDS", UE::StructUtils::FCreateUserDefinedStructArgs());
 
 			UTEST_TRUE("UUserDefinedStruct created", Result.HasValue());
 			UserDefinedStruct = Result.GetValue();
diff --git a/Engine/Source/Editor/StructUtilsEditor/Public/StructUtilsEditorUtils.h b/Engine/Source/Editor/StructUtilsEditor/Public/StructUtilsEditorUtils.h
--- a/Engine/Source/Editor/StructUtilsEditor/Public/StructUtilsEditorUtils.h
+++ b/Engine/Source/Editor/StructUtilsEditor/Public/StructUtilsEditorUtils.h
@@ -13,12 +13,25 @@ struct FPropertyBagPropertyDesc;
 
 namespace UE::StructUtils
 {
+	struct FCreateUserDefinedStructArgs
+	{
+		/** The default UserDefinedStruct flags. */
+		EObjectFlags UserDefinedStructFlags = RF_Transactional | RF_Public;
+		/** Add the BlueprintType metadata to the struct. */
+		bool bIsBlueprintType = true;
+		/**
+		 * The UserDefinedStruct editor appends the property ID to the property name.
+		 * Whether the newly created User Defined Struct should reproduce the same behavior.
+		 */
+		bool bAddPropertyIDToPropertyName = false;
+	};
+
+
 	/**
 	 * @return a new user defined struct created from a property bag.
 	 * Fails if the name is already used by another asset or if the asset could not be created.
-	 * When the property bag is empty, no user defined struct will be created.
 	 */
-	UE_API TValueOrError<UUserDefinedStruct*, void> CreateUserDefinedStructFromDescs(TNotNull<UObject*> UserDefinedStructOuter, const FInstancedPropertyBag& PropertyBag, FName UserDefinedStructName);
+	UE_API TValueOrError<UUserDefinedStruct*, void> CreateUserDefinedStructFromDescs(TNotNull<UObject*> UserDefinedStructOuter, const FInstancedPropertyBag& PropertyBag, FName UserDefinedStructName, const FCreateUserDefinedStructArgs& Args);
 
 	/**
 	 * Update, if needed, an existing user defined struct to have the exact same properties and default value as the property bag.
