[StructUtils] Fix unit test for reinstanciation. Introduced with CL

#rnx

#Synced-CL 49855772

diff --git a/Engine/Source/Editor/StructUtilsEditor/Private/Tests/StructReinstancerTests.cpp b/Engine/Source/Editor/StructUtilsEditor/Private/Tests/StructReinstancerTests.cpp
--- a/Engine/Source/Editor/StructUtilsEditor/Private/Tests/StructReinstancerTests.cpp
+++ b/Engine/Source/Editor/StructUtilsEditor/Private/Tests/StructReinstancerTests.cpp
@@ -109,7 +109,10 @@ namespace UE::StructUtils::Private
 
 		UUserDefinedStruct* UserDefinedStruct = nullptr;
 		{
-			TValueOrError<UUserDefinedStruct*, void> Result = UE::StructUtils::CreateUserDefinedStructFromDescs(UDSContainer.Get(), BagA, "UDS", UE::StructUtils::FCreateUserDefinedStructArgs());
+			UE::StructUtils::FCreateUserDefinedStructArgs Args;
+			Args.bAddPropertyIDToPropertyName = true;
+			Args.bIsBlueprintType = false;
+			TValueOrError<UUserDefinedStruct*, void> Result = UE::StructUtils::CreateUserDefinedStructFromDescs(UDSContainer.Get(), BagA, "UDS", Args);
 
 			UTEST_TRUE("UUserDefinedStruct created", Result.HasValue());
 			UserDefinedStruct = Result.GetValue();
@@ -245,7 +248,10 @@ namespace UE::StructUtils::Private
 
 		UUserDefinedStruct* UserDefinedStruct = nullptr;
 		{
-			TValueOrError<UUserDefinedStruct*, void> Result = UE::StructUtils::CreateUserDefinedStructFromDescs(UDSContainer.Get(), BagA, "UDS", UE::StructUtils::FCreateUserDefinedStructArgs());
+			UE::StructUtils::FCreateUserDefinedStructArgs Args;
+			Args.bAddPropertyIDToPropertyName = true;
+			Args.bIsBlueprintType = false;
+			TValueOrError<UUserDefinedStruct*, void> Result = UE::StructUtils::CreateUserDefinedStructFromDescs(UDSContainer.Get(), BagA, "UDS", Args);
 
 			UTEST_TRUE("UUserDefinedStruct created", Result.HasValue());
 			UserDefinedStruct = Result.GetValue();
@@ -367,7 +373,10 @@ namespace UE::StructUtils::Private
 		// Create UDS
 		UUserDefinedStruct* UserDefinedStruct = nullptr;
 		{
-			TValueOrError<UUserDefinedStruct*, void> Result = UE::StructUtils::CreateUserDefinedStructFromDescs(UDSContainer.Get(), BagA, "UDS", UE::StructUtils::FCreateUserDefinedStructArgs());
+			UE::StructUtils::FCreateUserDefinedStructArgs Args;
+			Args.bAddPropertyIDToPropertyName = true;
+			Args.bIsBlueprintType = false;
+			TValueOrError<UUserDefinedStruct*, void> Result = UE::StructUtils::CreateUserDefinedStructFromDescs(UDSContainer.Get(), BagA, "UDS", Args);
 
 			UTEST_TRUE("UUserDefinedStruct created", Result.HasValue());
 			UserDefinedStruct = Result.GetValue();
