[StructUtils] Fix typo. Introduced with CL 49924265.
#rnx
#Synced-CL 49855772

diff --git a/Engine/Source/Editor/StructUtilsEditor/Public/StructUtilsMetadata.h b/Engine/Source/Editor/StructUtilsEditor/Public/StructUtilsMetadata.h
--- a/Engine/Source/Editor/StructUtilsEditor/Public/StructUtilsMetadata.h
+++ b/Engine/Source/Editor/StructUtilsEditor/Public/StructUtilsMetadata.h
@@ -61,7 +61,7 @@ namespace UE::StructUtils::Metadata
 	/// [PropertyMetadata] StructTypeConst: Determine whether the type of the instanced struct can be changed
 	extern STRUCTUTILSEDITOR_API const FLazyName StructTypeConstName;
 
-	/// [PropertyMetadata] HideStructType: Determine whether the type of the instanced struct is visibled
+	/// [PropertyMetadata] HideStructType: Determine whether the type of the instanced struct is visible
 	extern STRUCTUTILSEDITOR_API const FLazyName HideStructTypeName;
 
 	/** If organizing by category is enabled for this property. */
