[StructUtils] Make better ?get struct? to prevent failure on bad memory access.
#rnx

#Synced-CL 50200679

diff --git a/Engine/Source/Editor/StructUtilsEditor/Private/Tests/StructReinstancerTests.cpp b/Engine/Source/Editor/StructUtilsEditor/Private/Tests/StructReinstancerTests.cpp
--- a/Engine/Source/Editor/StructUtilsEditor/Private/Tests/StructReinstancerTests.cpp
+++ b/Engine/Source/Editor/StructUtilsEditor/Private/Tests/StructReinstancerTests.cpp
@@ -423,7 +423,7 @@ namespace UE::StructUtils::Private
 				return true;
 			};
 
-		auto TestPropertyCDValue = [this, TestInstancedObject](int32 Value, int32 Offset)
+		auto TestPropertyCDValue = [this, TestInstancedObject](int32 Value, int32 Index)
 			{
 				const UScriptStruct* ScriptStruct = TestInstancedObject->InstancedStruct.GetScriptStruct();
 				const FStructProperty* StructProperty = CastField<const FStructProperty>(ScriptStruct->FindPropertyByName("MyStruct"));
@@ -432,12 +432,21 @@ namespace UE::StructUtils::Private
 					return false;
 				}
 
-				int32 CStructValue1 = 0;
-				StructProperty->GetValue_InContainer(TestInstancedObject->InstancedStruct.GetMemory() + Offset, &CStructValue1);
-				int32 CStructValue2 = 0;
-				StructProperty->GetValue_InContainer(TestInstancedObject->InstancedStructContainer[0].GetMemory() + Offset, &CStructValue2);
+				// The struct is 2 integers (C and D).
+				void* StructData = FMemory::Malloc(StructProperty->Struct->GetStructureSize());
+				StructProperty->Struct->InitializeStruct(StructData);
+
+				StructProperty->GetValue_InContainer(TestInstancedObject->InstancedStruct.GetMemory(), StructData);
+				int32 CStructValue1 = reinterpret_cast<int32*>(StructData)[Index];
+				StructProperty->GetValue_InContainer(TestInstancedObject->InstancedStructContainer[0].GetMemory(), StructData);
+				int32 CStructValue2 = reinterpret_cast<int32*>(StructData)[Index];
+
+				StructProperty->Struct->DestroyStruct(StructData);
+				FMemory::Free(StructData);
+
 				UTEST_TRUE(TEXT("TestInstancedObject.InstancedStruct"), Value == CStructValue1);
 				UTEST_TRUE(TEXT("TestInstancedObject.Container.InstancedStruct"), Value == CStructValue2);
+
 				return true;
 			};
 
@@ -504,7 +513,7 @@ namespace UE::StructUtils::Private
 			UTEST_TRUE(TEXT("Modified Test Values"), TestValue(AName, 1));
 			UTEST_TRUE(TEXT("Modified Test Values"), TestValue(BName, 2));
 			UTEST_TRUE(TEXT("Modified Test Values"), TestPropertyCDValue(-3, 0));
-			UTEST_TRUE(TEXT("Modified Test Values"), TestPropertyCDValue(4, sizeof(int32)));
+			UTEST_TRUE(TEXT("Modified Test Values"), TestPropertyCDValue(4, 1));
 		}
 
 		return true;
