Added Data Hierarchy Editor support to FInstancedPropertyBag
-- PropertyBagSchema class can now direct the property bag details to an optional HierarchyRoot object
--- PropertyBagSchema class also contains information about hierarchy types
--- This can be used to subclass and add properties
-- Moved some util functions around hash calculation into public
-- Replaced some UEdGraphSchema with UPropertyBagSchema
- Added more encapsulation so that FPropertyBagInstanceDataDetails is more standalone and can be subclassed
-- Removed some of the construct params in favor of FPropertyBagInstanceDataDetails looking relevant data up via property handle.
-- Added InstanceMetaData support to allow forwarding metadata to an inner FInstancedPropertyBag

#Synced-CL 50298389

diff --git a/Engine/Source/Editor/EditorStyle/Private/StarshipStyle.cpp b/Engine/Source/Editor/EditorStyle/Private/StarshipStyle.cpp
--- a/Engine/Source/Editor/EditorStyle/Private/StarshipStyle.cpp
+++ b/Engine/Source/Editor/EditorStyle/Private/StarshipStyle.cpp
@@ -473,6 +473,7 @@ void FStarshipEditorStyle::FStyle::SetupGeneralStyles()
 		Set("Icons.Heart", new IMAGE_BRUSH_SVG("Starship/Common/Heart", Icon16x16));
 		Set("Icons.RecentAssets", new IMAGE_BRUSH_SVG("Starship/Common/RecentAssets", Icon16x16));
 		Set("Icons.Scalability", new IMAGE_BRUSH_SVG("Starship/Common/Scalability_16", Icon20x20));
+		Set("Icons.EditHierarchy", new IMAGE_BRUSH_SVG("Starship/Common/EditHierarchy_16", Icon16x16));
 		Set("Icons.ViewportScalability", new IMAGE_BRUSH_SVG("Starship/Common/ViewportScalability_16", Icon16x16));
 		Set("Icons.ViewportScalabilityReset", new IMAGE_BRUSH_SVG("Starship/Common/ViewportScalabilityReset_16", Icon16x16));
 		Set("Icons.EllipsisVerticalNarrow",
diff --git a/Engine/Source/Editor/StructUtilsEditor/Private/Customizations/PropertyBagDetails.cpp b/Engine/Source/Editor/StructUtilsEditor/Private/Customizations/PropertyBagDetails.cpp
--- a/Engine/Source/Editor/StructUtilsEditor/Private/Customizations/PropertyBagDetails.cpp
+++ b/Engine/Source/Editor/StructUtilsEditor/Private/Customizations/PropertyBagDetails.cpp
@@ -2,30 +2,37 @@
 
 #include "PropertyBagDetails.h"
 
+#include "DataHierarchyViewModelBase.h"
 #include "StructUtilsEditorUtilsPrivate.h"
 
 #include "DetailLayoutBuilder.h"
 #include "DetailWidgetRow.h"
+#include "Editor.h"
 #include "IDetailChildrenBuilder.h"
 #include "IDetailGroup.h"
 #include "IPropertyUtilities.h"
-#include "PropertyCustomizationHelpers.h"
 #include "PropertyBagDragDropHandler.h"
+#include "PropertyCustomizationHelpers.h"
 #include "ScopedTransaction.h"
 #include "SDraggableBox.h" // Custom widget for drag and drop sections
 #include "StructUtilsMetadata.h"
 #include "STypeSelector.h" // Custom widget for pill type selector
+#include "Framework/Application/SlateApplication.h"
 #include "Framework/MultiBox/MultiBoxBuilder.h"
+#include "HierarchyEditor/PropertyBagHierarchyViewModel.h"
+#include "HierarchyEditor/SPropertyBagHierarchyEditor.h"
+#include "Kismet2/BlueprintEditorUtils.h"
 #include "StructUtils/UserDefinedStruct.h"
-#include "Templates/ValueOrError.h"
 #include "UObject/EnumProperty.h"
-#include "UObject/TextProperty.h"
+#include "Layout/Visibility.h"
 
+#include "Widgets/SDataHierarchyEditor.h"
+#include "Widgets/Docking/SDockTab.h"
 #include "Widgets/Images/SImage.h"
 #include "Widgets/Input/SButton.h"
 #include "Widgets/Input/SComboButton.h"
-#include "Widgets/Input/SEditableTextBox.h"
 #include "Widgets/Input/SMultiLineEditableTextBox.h"
+#include "Widgets/Layout/SWrapBox.h"
 #include "Widgets/Text/SInlineEditableTextBlock.h"
 
 #include UE_INLINE_GENERATED_CPP_BY_NAME(PropertyBagDetails)
@@ -34,6 +41,9 @@
 
 ////////////////////////////////////
 
+TWeakPtr<SPropertyBagHierarchyEditor> FPropertyBagDetails::PropertyBagHierarchyEditor;
+TWeakPtr<SWindow> FPropertyBagDetails::PropertyBagHierarchyEditorWindow;
+
 namespace UE::StructUtils
 {
 /** Sets property descriptor based on a Blueprint pin type. */
@@ -595,7 +605,6 @@ void DeleteProperty(TSharedPtr<IPropertyHandle> InStructProperty, TSharedPtr<IPr
 	ApplyChangesToPropertyDescs(
 		FText::Format(LOCTEXT("OnPropertyDeleted", "Deleted property: {0}"), ChildPropertyHandle->GetPropertyDisplayName()),
 		InStructProperty,
-		PropUtils,
 		[&ChildPropertyHandle](TArray<FPropertyBagPropertyDesc>& PropertyDescs)
 		{
 			const FProperty* Property = ChildPropertyHandle ? ChildPropertyHandle->GetProperty() : nullptr;
@@ -633,7 +642,6 @@ void PinInfoChanged(TSharedPtr<IPropertyHandle> ChildPropertyHandle, const TShar
 	ApplyChangesToPropertyDescs(
 		FText::Format(LOCTEXT("OnPropertyTypeChanged", "Changed property type: {0}"), ChildPropertyHandle->GetPropertyDisplayName()),
 		InBagStructProperty,
-		InPropUtils,
 		[&PinType, &ChildPropertyHandle](TArray<FPropertyBagPropertyDesc>& PropertyDescs)
 		{
 			// Find and change struct type
@@ -668,7 +676,7 @@ FOptionalSize GetDesiredInputWidgetSizeByType(const TSharedPtr<IPropertyHandle>&
 
 TSharedRef<SWidget> CreateTypeSelectionWidget(TSharedPtr<IPropertyHandle> ChildPropertyHandle, const TSharedPtr<IPropertyHandle>& InBagStructProperty, const TSharedPtr<IPropertyUtilities>& InPropUtils, SPinTypeSelector::ESelectorType SelectorType, const bool bAllowContainers, TSubclassOf<UEdGraphSchema> PropertyBagSchemaClass)
 {
-	const UEdGraphSchema* Schema = PropertyBagSchemaClass ? GetDefault<UEdGraphSchema>(PropertyBagSchemaClass) : GetDefault<UPropertyBagSchema>();
+	const UPropertyBagSchema* Schema = PropertyBagSchemaClass ? GetDefault<UPropertyBagSchema>(PropertyBagSchemaClass) : GetDefault<UPropertyBagSchema>();
 
 	return SNew(SBox)
 		.HAlign(HAlign_Right)
@@ -724,9 +732,65 @@ const UPropertyBag* GetCommonBagStruct(TSharedPtr<IPropertyHandle> StructPropert
 	return CommonBagStruct;
 }
 
-void ApplyChangesToPropertyDescs(const FText& SessionName, const TSharedPtr<IPropertyHandle>& StructProperty, const TSharedPtr<IPropertyUtilities>& PropUtils, TFunctionRef<void(TArray<FPropertyBagPropertyDesc>&)> Function)
+TSubclassOf<UPropertyBagSchema> ExtractPropertyBagSchemaClass(TSharedRef<IPropertyHandle> InPropertyBagHandle)
+{
+	FString PropertyBagSchemaClassName;
+	
+	/** We first try to extract the instance metadata, if specified. */
+	if (const FString* InstancePropertyBagSchemaClassName = InPropertyBagHandle->GetInstanceMetaData(UE::StructUtils::Metadata::PropertyBagSchemaClassName))
+	{
+		PropertyBagSchemaClassName = *InstancePropertyBagSchemaClassName;
+	}
+	
+	/** If not available, we attempt to use the class metadata. */
+	if (PropertyBagSchemaClassName.IsEmpty())
+	{
+		if (FProperty* MetaDataProperty = InPropertyBagHandle->GetMetaDataProperty())
+		{
+			PropertyBagSchemaClassName = MetaDataProperty->GetMetaData(UE::StructUtils::Metadata::PropertyBagSchemaClassName);
+		}
+	}
+	
+	if (!PropertyBagSchemaClassName.IsEmpty())
+	{
+		return FindObject<UClass>(nullptr, *PropertyBagSchemaClassName);
+	}
+	
+	return nullptr;
+}
+
+const UPropertyBagSchema* ExtractPropertyBagSchemaCDO(TSharedRef<IPropertyHandle> InPropertyBagHandle)
+{
+	TSubclassOf<UPropertyBagSchema> SchemaClass = ExtractPropertyBagSchemaClass(InPropertyBagHandle);
+	
+	if (const UPropertyBagSchema* Schema = SchemaClass ? GetDefault<UPropertyBagSchema>(SchemaClass) : GetDefault<UPropertyBagSchema>())
+	{
+		return Schema;
+	}
+	
+	return nullptr;
+}
+
+UPropertyBagHierarchyRoot* ExtractHierarchyRoot(TSharedRef<IPropertyHandle> InPropertyBagHandle)
+{
+	if (!InPropertyBagHandle->IsValidHandle())
+	{
+		return nullptr;
+	}
+	
+	if (const UPropertyBagSchema* Schema = ExtractPropertyBagSchemaCDO(InPropertyBagHandle))
+	{
+		TArray<UObject*> OuterObjects;
+		InPropertyBagHandle->GetOuterObjects(OuterObjects);
+		return Schema->GetHierarchyRoot(OuterObjects);
+	}
+	
+	return nullptr;
+}
+
+void ApplyChangesToPropertyDescs(const FText& SessionName, const TSharedPtr<IPropertyHandle>& StructProperty, TFunctionRef<void(TArray<FPropertyBagPropertyDesc>&)> Function)
 {
-	if (!StructProperty || !StructProperty.Get()->IsValidHandle() || !PropUtils)
+	if (!StructProperty || !StructProperty.Get()->IsValidHandle())
 	{
 		return;
 	}
@@ -734,39 +798,52 @@ void ApplyChangesToPropertyDescs(const FText& SessionName, const TSharedPtr<IPro
 	FScopedTransaction Transaction(SessionName);
 	TArray<FPropertyBagPropertyDesc> PropertyDescs = Private::GetCommonPropertyDescs(StructProperty);
 	StructProperty->NotifyPreChange();
-
+	
+	if (UPropertyBagHierarchyRoot* Root = ExtractHierarchyRoot(StructProperty.ToSharedRef()))
+	{
+		Root->Modify();
+	}
+	
 	Function(PropertyDescs);
 
 	Private::SetPropertyDescs(StructProperty, PropertyDescs);
 
 	StructProperty->NotifyPostChange(EPropertyChangeType::ValueSet);
 	StructProperty->NotifyFinishedChangingProperties();
-	if (PropUtils)
-	{
-		PropUtils->ForceRefresh();
-	}
+	
+	StructProperty->RequestRebuildChildren();
+}
+
+void ApplyChangesToPropertyDescs(const FText& SessionName, const TSharedPtr<IPropertyHandle>& StructProperty, const TSharedPtr<IPropertyUtilities>& PropUtils, TFunctionRef<void(TArray<FPropertyBagPropertyDesc>&)> Function)
+{
+	ApplyChangesToPropertyDescs(SessionName, StructProperty, Function);
 }
 
-void ApplyChangesToSinglePropertyDesc(const FText& SessionName, const TSharedPtr<IPropertyHandle> PropertyHandle, const TSharedPtr<IPropertyHandle>& StructProperty, const TSharedPtr<IPropertyUtilities>& PropUtils, TFunctionRef<void(FPropertyBagPropertyDesc&)> Function)
+void ApplyChangesToSinglePropertyDesc(const FText& SessionName, const TSharedPtr<IPropertyHandle> PropertyHandle, const TSharedPtr<IPropertyHandle>& StructProperty, TFunctionRef<void(FPropertyBagPropertyDesc&)> Function)
 {
-	ApplyChangesToPropertyDescs(SessionName, StructProperty, PropUtils, [Function = std::move(Function), Property = PropertyHandle->GetProperty()](TArray<FPropertyBagPropertyDesc>& PropertyDescs)
-	{
-		if (FPropertyBagPropertyDesc* Desc = PropertyDescs.FindByPredicate([Property](const FPropertyBagPropertyDesc& OutDesc) { return OutDesc.CachedProperty == Property; }))
+	ApplyChangesToPropertyDescs(SessionName, StructProperty, [Function = std::move(Function), Property = PropertyHandle->GetProperty()](TArray<FPropertyBagPropertyDesc>& PropertyDescs)
 		{
-			Function(*Desc);
-		}
-	});
+			if (FPropertyBagPropertyDesc* Desc = PropertyDescs.FindByPredicate([Property](const FPropertyBagPropertyDesc& OutDesc) { return OutDesc.CachedProperty == Property; }))
+			{
+				Function(*Desc);
+			}
+		});
 }
 
-void ApplyChangesToSinglePropertyDesc(const FText& SessionName, const FPropertyBagPropertyDesc& PropertyDesc, const TSharedPtr<IPropertyHandle>& StructProperty, const TSharedPtr<IPropertyUtilities>& PropUtils, TFunctionRef<void(FPropertyBagPropertyDesc&)> Function)
+void ApplyChangesToSinglePropertyDesc(const FText& SessionName, const FPropertyBagPropertyDesc& PropertyDesc, const TSharedPtr<IPropertyHandle>& StructProperty, TFunctionRef<void(FPropertyBagPropertyDesc&)> Function)
 {
-	ApplyChangesToPropertyDescs(SessionName, StructProperty, PropUtils, [Function = std::move(Function), &PropertyDesc](TArray<FPropertyBagPropertyDesc>& PropertyDescs)
-	{
-		if (FPropertyBagPropertyDesc* Desc = PropertyDescs.FindByPredicate([&PropertyDesc](const FPropertyBagPropertyDesc& OutDesc) { return OutDesc == PropertyDesc; }))
+	ApplyChangesToPropertyDescs(SessionName, StructProperty, [Function = std::move(Function), &PropertyDesc](TArray<FPropertyBagPropertyDesc>& PropertyDescs)
 		{
-			Function(*Desc);
-		}
-	});
+			if (FPropertyBagPropertyDesc* Desc = PropertyDescs.FindByPredicate([&PropertyDesc](const FPropertyBagPropertyDesc& OutDesc) { return OutDesc == PropertyDesc; }))
+			{
+				Function(*Desc);
+			}
+		});
+}
+
+void ApplyChangesToSinglePropertyDesc(const FText& SessionName, const FPropertyBagPropertyDesc& PropertyDesc, const TSharedPtr<IPropertyHandle>& StructProperty, const TSharedPtr<IPropertyUtilities>& PropUtils, TFunctionRef<void(FPropertyBagPropertyDesc&)> Function)
+{
+	ApplyChangesToSinglePropertyDesc(SessionName, PropertyDesc, StructProperty, Function);
 }
 
 namespace Constants
@@ -781,6 +858,85 @@ namespace Constants
 }
 } // UE::StructUtils
 
+
+/** The category builder will display all sub-categories & parameters contained within a given category in the details panel. */
+class FHierarchyPropertyCategoryBuilder : public IDetailCustomNodeBuilder, public TSharedFromThis<FHierarchyPropertyCategoryBuilder>
+{
+public:
+	FHierarchyPropertyCategoryBuilder(const UPropertyBagHierarchyCategory& InHierarchyCategory, TSharedRef<FPropertyBagInstanceDataDetails> InOwningDetails);
+	
+	virtual void GenerateHeaderRowContent(FDetailWidgetRow& NodeRow) override;
+	
+	virtual bool InitiallyCollapsed() const override;
+
+	virtual FName GetName() const  override;
+
+	virtual void GenerateChildContent(IDetailChildrenBuilder& ChildrenBuilder) override;
+
+private:
+	TWeakObjectPtr<const UPropertyBagHierarchyCategory> HierarchyCategory;
+	FName CustomBuilderRowName;
+	TWeakPtr<FPropertyBagInstanceDataDetails> OwningDetails;
+};
+
+FHierarchyPropertyCategoryBuilder::FHierarchyPropertyCategoryBuilder(const UPropertyBagHierarchyCategory& InHierarchyCategory, TSharedRef<FPropertyBagInstanceDataDetails> InOwningDetails)
+	: HierarchyCategory(&InHierarchyCategory)
+	, CustomBuilderRowName(HierarchyCategory->GetCategoryName())
+	, OwningDetails(InOwningDetails)
+{	
+}
+
+void FHierarchyPropertyCategoryBuilder::GenerateHeaderRowContent(FDetailWidgetRow& NodeRow)
+{
+	NodeRow.NameContent()
+	[
+		SNew(STextBlock)
+		.Text_UObject(HierarchyCategory.Get(), &UHierarchyCategory::GetCategoryAsText)
+		.ToolTipText_UObject(HierarchyCategory.Get(), &UHierarchyCategory::GetTooltip)
+	];
+
+	NodeRow.FilterString(FText::FromName(CustomBuilderRowName));
+}
+
+bool FHierarchyPropertyCategoryBuilder::InitiallyCollapsed() const
+{
+	return false;
+}
+
+FName FHierarchyPropertyCategoryBuilder::GetName() const
+{
+	return CustomBuilderRowName;
+}
+
+void FHierarchyPropertyCategoryBuilder::GenerateChildContent(IDetailChildrenBuilder& ChildrenBuilder)
+{
+	if (!OwningDetails.IsValid())
+	{
+		return;
+	}
+	
+	TArray<const UPropertyBagHierarchyCategory*> ContainedCategories;
+	HierarchyCategory->GetChildrenOfType(ContainedCategories, false);
+	
+	for (const UPropertyBagHierarchyCategory* ContainedHierarchyCategory : ContainedCategories)
+	{
+		ChildrenBuilder.AddCustomBuilder(MakeShared<FHierarchyPropertyCategoryBuilder>(*ContainedHierarchyCategory, OwningDetails.Pin().ToSharedRef()));
+	}
+		
+	TArray<const UPropertyBagHierarchyProperty*> ContainedProperties;
+	HierarchyCategory->GetChildrenOfType<>(ContainedProperties, false);
+	
+	for (const UPropertyBagHierarchyProperty* ContainedHierarchyProperty : ContainedProperties)
+	{
+		if (OwningDetails.Pin()->GetHierarchyPropertyHandleMap().Contains(ContainedHierarchyProperty))
+		{
+			IDetailPropertyRow& PropertyRow = ChildrenBuilder.AddProperty(OwningDetails.Pin()->GetHierarchyPropertyHandleMap()[ContainedHierarchyProperty].ToSharedRef());
+			PropertyRow.Visibility(MakeAttributeSP(OwningDetails.Pin().Get(), &FPropertyBagInstanceDataDetails::ShouldShowHierarchyProperty, PropertyRow.GetPropertyHandle()));
+			OwningDetails.Pin()->OnChildRowAdded(PropertyRow);
+		}
+	}
+}
+
 //----------------------------------------------------------------//
 //  FPropertyBagInstanceDataDetails
 //  - StructProperty is FInstancedPropertyBag
@@ -798,10 +954,21 @@ PRAGMA_DISABLE_DEPRECATION_WARNINGS
 	, bFixedLayout(false)
 	, bAllowArrays(true)
 PRAGMA_ENABLE_DEPRECATION_WARNINGS
-	, PropertyBagSchemaClass(ConstructParams.PropertyBagSchemaClass)
 {
 	ensure(UE::StructUtils::Private::IsScriptStruct<FInstancedPropertyBag>(BagStructProperty));
 	ensure(PropUtils != nullptr);
+	
+	if (BagStructProperty.IsValid() && BagStructProperty->IsValidHandle())
+	{
+		PropertyBagSchema = UE::StructUtils::ExtractPropertyBagSchemaCDO(BagStructProperty.ToSharedRef());
+		
+		HierarchyData.HierarchyRoot = UE::StructUtils::ExtractHierarchyRoot(BagStructProperty.ToSharedRef());
+		
+		if (HierarchyData.HierarchyRoot.IsValid())
+		{
+			HierarchyData.OnNavigateToHierarchyPropertyRequestedDelegate = ConstructParams.HierarchyParams.OnNavigateToHierarchyPropertyRequested;
+		}
+	}
 }
 
 /** For backwards compatibility. */
@@ -820,12 +987,93 @@ PRAGMA_ENABLE_DEPRECATION_WARNINGS
 	ensure(PropUtils != nullptr);
 }
 
+void FPropertyBagInstanceDataDetails::GenerateHierarchyShowAdvancedRow(IDetailChildrenBuilder& ChildrenBuilder)
+{
+	if (!HierarchyData.HierarchyRoot.IsValid())
+	{
+		return;
+	}
+	
+	// At last we add a row for the 'Show Advanced' button. Only displayed if there is at least one advanced property
+	FDetailWidgetRow& AdvancedWidgetRow = ChildrenBuilder.AddCustomRow(INVTEXT("Advanced"));
+	AdvancedWidgetRow.Visibility(TAttribute<EVisibility>::CreateSP(this, &FPropertyBagInstanceDataDetails::GetAdvancedButtonVisibility));
+		
+	AdvancedWidgetRow.WholeRowContent()
+	.HAlign(HAlign_Center)
+	[
+		SAssignNew(ShowAdvancedButton, SButton)
+		.ButtonStyle(FAppStyle::Get(), "HoverHintOnly")
+		.ToolTipText(LOCTEXT("DisplayedAdvancedProperties", "Display advanced properties"))
+		.OnClicked_Lambda([this]()
+		{
+			bShouldShowAdvanced = !bShouldShowAdvanced;
+			return FReply::Handled();
+		})
+		[
+			SNew(SImage)
+			.Image(this, &FPropertyBagInstanceDataDetails::GetAdvancedImage)
+		]
+	];
+}
+
+void FPropertyBagInstanceDataDetails::AddChildRows(IDetailChildrenBuilder& ChildBuilder, const TArray<TSharedPtr<IPropertyHandle>>& ChildProperties)
+{
+	if (HierarchyData.HierarchyRoot.IsValid() && BagStructProperty.IsValid() && BagStructProperty->IsValidHandle())
+	{		
+		void* BagStructPtr = nullptr;
+		FPropertyAccess::Result Result = BagStructProperty->GetValueData(BagStructPtr);
+		
+		if (Result == FPropertyAccess::Fail)
+		{
+			return;
+		}
+		
+		FInstancedPropertyBag* Bag = static_cast<FInstancedPropertyBag*>(BagStructPtr);
+		
+		// We build a map from hierarchy property to property handle. This is then used in hierarchy category builders to retrieve their handles to add to the details panel.
+		TArray<const UPropertyBagHierarchyProperty*> AllHierarchyProperties;
+		HierarchyData.HierarchyRoot->GetChildrenOfType(AllHierarchyProperties, true);
+		
+		HierarchyData.HierarchyPropertyHandleMap.Empty();
+		for(const UPropertyBagHierarchyProperty* HierarchyProperty : AllHierarchyProperties)
+		{
+			if (const FPropertyBagPropertyDesc* PropertyDesc = Bag->FindPropertyDescByID(HierarchyProperty->GetPropertyId()))
+			{
+				const TSharedPtr<IPropertyHandle>* MatchingPropertyHandle = ChildProperties.FindByPredicate([PropertyDesc](const TSharedPtr<IPropertyHandle>& Candidate)
+				{
+					return Candidate->GetProperty() == PropertyDesc->CachedProperty;
+				});
+				
+				if (MatchingPropertyHandle)
+				{
+					HierarchyData.HierarchyPropertyHandleMap.Add(HierarchyProperty, *MatchingPropertyHandle);			
+				}
+			}
+		}
+		
+		// First, we generate the section row
+		GenerateHierarchySectionRow(ChildBuilder);
+		// Then the root level categories & properties
+		GenerateHierarchyRootRows(ChildBuilder);
+		// Then the leftover properties that have not been embedded into the hierarchy
+		GenerateHierarchyLeftoverPropertyRows(ChildBuilder, ChildProperties);
+		// Then the 'Show Advanced' row, if there is any advanced property
+		GenerateHierarchyShowAdvancedRow(ChildBuilder);
+	}
+	else
+	{
+		FInstancedStructDataDetails::AddChildRows(ChildBuilder, ChildProperties);
+	}
+}
+
 void FPropertyBagInstanceDataDetails::OnGroupRowAdded(IDetailGroup& GroupRow, int32 Level, const FString& Category) const
 {
 	using namespace UE::StructUtils;
 
 	FDetailWidgetRow& FolderRow = GroupRow.HeaderRow();
 	TWeakPtr<const FPropertyBagInstanceDataDetails> WeakSelf = SharedThis<const FPropertyBagInstanceDataDetails>(this);
+	TWeakPtr<IPropertyHandle> WeakPropertyBagHandle = BagStructProperty;
+	
 	FString FullCategoryName = GroupRow.GetGroupName().ToString();
 
 	/*** DRAG AND DROP HANDLER ***/
@@ -862,8 +1110,13 @@ void FPropertyBagInstanceDataDetails::OnGroupRowAdded(IDetailGroup& GroupRow, in
 
 		DragDropHandler->BindOnHandleDragDrop(
 			FOnPropertyBagDetailsRowDropOp::CreateLambda(
-				[WeakSelf, FullCategoryName, BagStructProperty = BagStructProperty, PropUtils = PropUtils](const FPropertyBagPropertyDesc& DroppedPropertyDesc, const EItemDropZone DropZone) -> FReply
+				[WeakSelf, FullCategoryName, BagStructProperty = WeakPropertyBagHandle, PropUtils = PropUtils](const FPropertyBagPropertyDesc& DroppedPropertyDesc, const EItemDropZone DropZone) -> FReply
 				{
+					if (!BagStructProperty.IsValid())
+					{
+						return FReply::Unhandled();
+					}
+						
 					if (ensure(DroppedPropertyDesc.CachedProperty && DropZone == EItemDropZone::OntoItem))
 					{
 						const TSharedPtr<const FPropertyBagInstanceDataDetails> DetailsSP = WeakSelf.Pin();
@@ -877,8 +1130,7 @@ void FPropertyBagInstanceDataDetails::OnGroupRowAdded(IDetailGroup& GroupRow, in
 						ApplyChangesToSinglePropertyDesc(
 							LOCTEXT("DragToChangeCategory", "Change property category"),
 							DroppedPropertyDesc,
-							BagStructProperty,
-							PropUtils,
+							BagStructProperty.Pin(),
 							[&DroppedPropertyDesc, &FullCategoryName](FPropertyBagPropertyDesc& Desc)
 							{
 								Metadata::SetCategory(Desc, FullCategoryName);
@@ -933,7 +1185,6 @@ void FPropertyBagInstanceDataDetails::OnGroupRowAdded(IDetailGroup& GroupRow, in
 				ApplyChangesToPropertyDescs(
 					LOCTEXT("InlineRenameCategory", "Rename category"),
 					BagStructProperty,
-					PropUtils,
 					[&InNewText, OldCategory = FullCategoryName, &Category](TArray<FPropertyBagPropertyDesc>& PropertyDescs)
 					{
 						FString NewCategory = OldCategory;
@@ -988,7 +1239,6 @@ void FPropertyBagInstanceDataDetails::OnGroupRowAdded(IDetailGroup& GroupRow, in
 					ApplyChangesToPropertyDescs(
 						LOCTEXT("OnCategoryDeleted", "Delete category"),
 						BagStructProperty,
-						PropUtils,
 						[GroupName](TArray<FPropertyBagPropertyDesc>& PropertyDescs)
 						{
 							for (FPropertyBagPropertyDesc& Desc : PropertyDescs)
@@ -1188,7 +1438,7 @@ void FPropertyBagInstanceDataDetails::OnChildRowAdded(IDetailPropertyRow& ChildR
 				.Padding(1, 0)
 				.AutoWidth()
 				[
-					CreateTypeSelectionWidget(ChildPropertyHandle, BagStructProperty, PropUtils, SPinTypeSelector::ESelectorType::Compact, bAllowContainers, PropertyBagSchemaClass)
+					CreateTypeSelectionWidget(ChildPropertyHandle, BagStructProperty, PropUtils, SPinTypeSelector::ESelectorType::Compact, bAllowContainers, PropertyBagSchema->GetClass())
 				];
 		}
 
@@ -1272,7 +1522,6 @@ void FPropertyBagInstanceDataDetails::OnChildRowAdded(IDetailPropertyRow& ChildR
 				ApplyChangesToPropertyDescs(
 					FText::Format(LOCTEXT("OnPropertyNameChanged", "Change property name: {0} -> {1}"), FText::FromName(OldName), FText::FromName(NewName)),
 					BagStructProperty,
-					PropUtils,
 					[&NewName, &ChildPropertyHandle](TArray<FPropertyBagPropertyDesc>& PropertyDescs)
 					{
 						const FProperty* Property = ChildPropertyHandle->GetProperty();
@@ -1353,7 +1602,6 @@ void FPropertyBagInstanceDataDetails::OnChildRowAdded(IDetailPropertyRow& ChildR
 							LOCTEXT("OnPropertyAccessSpecifierChanged", "Set access specifier."),
 							ChildPropertyHandle,
 							BagStructProperty,
-							PropUtils,
 							[Property, BagStructProperty](FPropertyBagPropertyDesc& Desc)
 							{
 								const bool bIsPrivate = !Property->HasAnyPropertyFlags(CPF_NativeAccessSpecifierPublic);
@@ -1411,7 +1659,7 @@ void FPropertyBagInstanceDataDetails::OnChildRowAdded(IDetailPropertyRow& ChildR
 						PropUtils,
 						SPinTypeSelector::ESelectorType::Full,
 						bAllowContainers,
-						PropertyBagSchemaClass),
+						PropertyBagSchema->GetClass()),
 					/*InLabel=*/FText::GetEmpty());
 				MenuBuilder.EndSection();
 			}
@@ -1454,145 +1702,114 @@ void FPropertyBagInstanceDataDetails::OnChildRowAdded(IDetailPropertyRow& ChildR
 
 				MenuBuilder.EndSection();
 			}
-
-			// BagStructProperty and PropUtils are forwarded straight to OnTextCommitted, so they save a ref count and be refs here.
-			auto BuildMetadataSpecifierSubMenu = [](
-				TAttribute<FText>&& TextAttribute,
-				FOnVerifyTextChanged&& OnVerifyTextDelegate,
-				FOnTextCommitted&& OnTextCommittedDelegate) -> FNewMenuDelegate
-			{
-				TSharedRef<SWidget> TextBox = SNew(SMultiLineEditableTextBox)
-					.WrapTextAt(400)
-					.Text(MoveTemp(TextAttribute))
-					.OnVerifyTextChanged(MoveTemp(OnVerifyTextDelegate))
-					.OnTextCommitted(OnTextCommittedDelegate);
-
-				return FNewMenuDelegate::CreateLambda(
-					[TextBox](FMenuBuilder& OutMenuBuilder)
-					{
-						OutMenuBuilder.AddWidget(
-							SNew(SHorizontalBox)
-							+ SHorizontalBox::Slot()
-							.HAlign(HAlign_Fill)
-							.MinWidth(Constants::MinSubMenuTextBoxWidth)
-							.MaxWidth(Constants::MaxSubMenuTextBoxWidth)
-							.Padding(Constants::SubMenuTextBoxPadding, 0)
-							.AutoWidth()
-							[
-								TextBox
-							],
-							FText::GetEmpty(),
-							/*bNoIndent=*/true,
-							/*bInSearchable=*/true);
-					});
-			};
-
-			// The property's category (grouping) can be edited here.
-			if (EnumHasAllFlags(ChildRowFeatures, EPropertyBagChildRowFeatures::Categories | EPropertyBagChildRowFeatures::Menu_Categories))
+			
+			if (HierarchyData.HierarchyRoot.IsValid())
 			{
-				MenuBuilder.BeginSection(/*InExtensionHook=*/NAME_None, LOCTEXT("DropDownMenuSectionCategory", "Category"));
-
-				if (ChildPropertyHandle->HasMetaData(Metadata::CategoryName))
+				if (HierarchyData.OnNavigateToHierarchyPropertyRequestedDelegate.IsBound())
 				{
 					MenuBuilder.AddMenuEntry(
-						LOCTEXT("DropDownMenuClearCategory", "Clear category"),
-						LOCTEXT("DropDownMenuClearCategoryToolTip", "Remove the property from its current category."),
-						FSlateIcon(FAppStyle::GetAppStyleSetName(), "Icons.Delete"),
-						FUIAction(FExecuteAction::CreateLambda([ChildPropertyHandle, StructProperty = BagStructProperty, PropUtils = PropUtils]()
+						LOCTEXT("DropDownMenuEditPropertyInHierarchyEditor", "Edit in Hierarchy Editor"),
+						LOCTEXT("DropDownMenUEditPropertyInHierarchyEditorToolTip", "Edit this property in the Hierarchy Editor"),
+						FSlateIcon(),
+						FUIAction(FExecuteAction::CreateLambda([this, PropertyDesc]()
 						{
-							ApplyChangesToSinglePropertyDesc(
-								LOCTEXT("DropDownMenuOnCategoryCleared", "Clear property category"),
-								ChildPropertyHandle,
-								StructProperty,
-								PropUtils,
-								[](FPropertyBagPropertyDesc& Desc)
-								{
-									Metadata::RemoveCategory(Desc);
-								});
-						})));
+							HierarchyData.OnNavigateToHierarchyPropertyRequestedDelegate.Execute(*PropertyDesc);
+						}))
+					);
 				}
+			}
 
-				FOnVerifyTextChanged ValidateCategoryNameDelegate = FOnVerifyTextChanged::CreateLambda([](const FText& InText, FText& OutMessage)
+			// We only add metadata menu entries if we don't have a hierarchy root
+			if (!HierarchyData.HierarchyRoot.IsValid())
+			{
+				// PropertyBagHandle and PropUtils are forwarded straight to OnTextCommitted, so they save a ref count and be refs here.
+				auto BuildMetadataSpecifierSubMenu = [](
+					TAttribute<FText>&& TextAttribute,
+					FOnVerifyTextChanged&& OnVerifyTextDelegate,
+					FOnTextCommitted&& OnTextCommittedDelegate) -> FNewMenuDelegate
 				{
-					if (InText.ToString().Len() > Constants::MaxCategoryLength)
-					{
-						OutMessage = LOCTEXT("DropDownMenuInvalidCategoryName", "Invalid category name");
-						return false;
-					}
-					else
-					{
-						return true;
-					}
-				});
-
-				MenuBuilder.AddSubMenu(
-					FText::FromName(Metadata::CategoryName),
-					LOCTEXT("DropDownMenuSubMenuCategoryTooltip", "Set the category of this property. Subcategories can be created with the '|' character."),
-					BuildMetadataSpecifierSubMenu(
-						TAttribute<FText>::CreateLambda([ChildPropertyHandle]()
+					TSharedRef<SWidget> TextBox = SNew(SMultiLineEditableTextBox)
+						.WrapTextAt(400)
+						.Text(MoveTemp(TextAttribute))
+						.OnVerifyTextChanged(MoveTemp(OnVerifyTextDelegate))
+						.OnTextCommitted(OnTextCommittedDelegate);
+
+					return FNewMenuDelegate::CreateLambda(
+						[TextBox](FMenuBuilder& OutMenuBuilder)
 						{
-							FText GroupLabel = FText::FromString("");
-							if (ChildPropertyHandle && ChildPropertyHandle->HasMetaData(Metadata::CategoryName))
-							{
-								GroupLabel = FText::FromString(ChildPropertyHandle->GetMetaData(Metadata::CategoryName));
-							}
+							OutMenuBuilder.AddWidget(
+								SNew(SHorizontalBox)
+								+ SHorizontalBox::Slot()
+								.HAlign(HAlign_Fill)
+								.MinWidth(Constants::MinSubMenuTextBoxWidth)
+								.MaxWidth(Constants::MaxSubMenuTextBoxWidth)
+								.Padding(Constants::SubMenuTextBoxPadding, 0)
+								.AutoWidth()
+								[
+									TextBox
+								],
+								FText::GetEmpty(),
+								/*bNoIndent=*/true,
+								/*bInSearchable=*/true);
+						});
+				};
 
-							return GroupLabel;
-						}),
-						FOnVerifyTextChanged::CreateLambda([](const FText& InText, FText& OutMessage)
-						{
-							if (InText.ToString().Len() > Constants::MaxCategoryLength)
-							{
-								OutMessage = LOCTEXT("DropDownMenuInvalidCategoryName", "Invalid category name");
-								return false;
-							}
-							else
-							{
-								return true;
-							}
-						}),
+				// The property's category (grouping) can be edited here.
+				if (EnumHasAllFlags(ChildRowFeatures, EPropertyBagChildRowFeatures::Categories | EPropertyBagChildRowFeatures::Menu_Categories))
+				{
+					MenuBuilder.BeginSection(/*InExtensionHook=*/NAME_None, LOCTEXT("DropDownMenuSectionCategory", "Category"));
 
-						FOnTextCommitted::CreateLambda([StructProperty = BagStructProperty, PropUtils = PropUtils, ChildPropertyHandle](const FText& CommittedText, const ETextCommit::Type CommitType)
-						{
-							if (CommitType == ETextCommit::OnEnter || CommitType == ETextCommit::OnUserMovedFocus)
+					if (ChildPropertyHandle->HasMetaData(Metadata::CategoryName))
+					{
+						MenuBuilder.AddMenuEntry(
+							LOCTEXT("DropDownMenuClearCategory", "Clear category"),
+							LOCTEXT("DropDownMenuClearCategoryToolTip", "Remove the property from its current category."),
+							FSlateIcon(FAppStyle::GetAppStyleSetName(), "Icons.Delete"),
+							FUIAction(FExecuteAction::CreateLambda([ChildPropertyHandle, StructProperty = BagStructProperty, PropUtils = PropUtils]()
 							{
 								ApplyChangesToSinglePropertyDesc(
-									LOCTEXT("DropDownMenuOnCategoryEdited", "Edit property category"),
+									LOCTEXT("DropDownMenuOnCategoryCleared", "Clear property category"),
 									ChildPropertyHandle,
 									StructProperty,
-									PropUtils,
-									[&CommittedText](FPropertyBagPropertyDesc& Desc)
-									{
-										Metadata::SetCategory(Desc, CommittedText.ToString());
-									});
-							}
-						})),
-					/*bInOpenSubMenuOnClick=*/true);
+									[](FPropertyBagPropertyDesc& Desc)
+										{
+											Metadata::RemoveCategory(Desc);
+										});
+							})));
+					}
 
-				MenuBuilder.EndSection();
-			}
+					FOnVerifyTextChanged ValidateCategoryNameDelegate = FOnVerifyTextChanged::CreateLambda([](const FText& InText, FText& OutMessage)
+					{
+						if (InText.ToString().Len() > Constants::MaxCategoryLength)
+						{
+							OutMessage = LOCTEXT("DropDownMenuInvalidCategoryName", "Invalid category name");
+							return false;
+						}
+						else
+						{
+							return true;
+						}
+					});
 
-			// The property's optional metadata specifiers can be edited here.
-			if (EnumHasAllFlags(ChildRowFeatures, EPropertyBagChildRowFeatures::Menu_MetadataSpecifiers))
-			{
-				MenuBuilder.BeginSection(/*InExtensionHook=*/NAME_None, LOCTEXT("DropDownMenuSectionMetadataSpecifiers", "Specifiers"));
+					MenuBuilder.AddSubMenu(
+						FText::FromName(Metadata::CategoryName),
+						LOCTEXT("DropDownMenuSubMenuCategoryTooltip", "Set the category of this property. Subcategories can be created with the '|' character."),
+						BuildMetadataSpecifierSubMenu(
+							TAttribute<FText>::CreateLambda([ChildPropertyHandle]()
+							{
+								FText GroupLabel = FText::FromString("");
+								if (ChildPropertyHandle && ChildPropertyHandle->HasMetaData(Metadata::CategoryName))
+								{
+									GroupLabel = FText::FromString(ChildPropertyHandle->GetMetaData(Metadata::CategoryName));
+								}
 
-				auto BuildGenericSpecifierSubMenu = [this, ChildPropertyHandle, BuildMetadataSpecifierSubMenu](const FName& SpecifierName, const bool bIsNumeric = false)
-				{
-					return BuildMetadataSpecifierSubMenu(
-						TAttribute<FText>::CreateLambda([SpecifierName, ChildPropertyHandle]()
-						{
-							return ChildPropertyHandle.IsValid()
-									   ? FText::FromString(ChildPropertyHandle->GetMetaData(SpecifierName))
-									   : FText::GetEmpty();
-						}),
-						!bIsNumeric
-							? FOnVerifyTextChanged()
-							: FOnVerifyTextChanged::CreateLambda([](const FText& InText, FText& OutMessage) -> bool
+								return GroupLabel;
+							}),
+							FOnVerifyTextChanged::CreateLambda([](const FText& InText, FText& OutMessage)
 							{
-								if (!InText.IsEmpty() && !InText.IsNumeric())
+								if (InText.ToString().Len() > Constants::MaxCategoryLength)
 								{
-									OutMessage = LOCTEXT("DropDownMenuSpecifierNumericValueExpected", "Value is not numeric");
+									OutMessage = LOCTEXT("DropDownMenuInvalidCategoryName", "Invalid category name");
 									return false;
 								}
 								else
@@ -1600,76 +1817,124 @@ void FPropertyBagInstanceDataDetails::OnChildRowAdded(IDetailPropertyRow& ChildR
 									return true;
 								}
 							}),
-						FOnTextCommitted::CreateLambda(
-							[SpecifierName, StructProperty = BagStructProperty, PropUtils = PropUtils, ChildPropertyHandle]
-							(const FText& CommittedText, const ETextCommit::Type CommitType)
+
+							FOnTextCommitted::CreateLambda([StructProperty = BagStructProperty, PropUtils = PropUtils, ChildPropertyHandle](const FText& CommittedText, const ETextCommit::Type CommitType)
 							{
-								if (CommitType != ETextCommit::OnCleared)
+								if (CommitType == ETextCommit::OnEnter || CommitType == ETextCommit::OnUserMovedFocus)
 								{
 									ApplyChangesToSinglePropertyDesc(
-										FText::Format(LOCTEXT("DropDownMenuOnSetSpecifier", "Set metadata specifier '{0}' value"), FText::FromName(SpecifierName)),
+										LOCTEXT("DropDownMenuOnCategoryEdited", "Edit property category"),
 										ChildPropertyHandle,
 										StructProperty,
-										PropUtils,
-										[SpecifierName, &CommittedText](FPropertyBagPropertyDesc& Desc)
-										{
-											// If empty, remove the specifier
-											if (CommittedText.IsEmpty())
-											{
-												Desc.RemoveMetadata(SpecifierName);
-											}
-											else
+										[&CommittedText](FPropertyBagPropertyDesc& Desc)
 											{
-												Desc.SetMetaData(SpecifierName, CommittedText.ToString());
-											}
-										});
+												Metadata::SetCategory(Desc, CommittedText.ToString());
+											});
 								}
-							}));
-				};
+							})),
+						/*bInOpenSubMenuOnClick=*/true);
 
-				MenuBuilder.AddSubMenu(
-					FText::FromName(Metadata::Specifiers::ToolTipName),
-					LOCTEXT("DropDownMenuSubMenuTooltipTooltip", "Set the tooltip for this property."),
-					BuildGenericSpecifierSubMenu(Metadata::Specifiers::ToolTipName),
-					/*bInOpenSubMenuOnClick=*/true);
+					MenuBuilder.EndSection();
+				}
 
-				if (Metadata::Specifiers::SupportsClamping(*PropertyDesc))
+				// The property's optional metadata specifiers can be edited here.
+				if (EnumHasAllFlags(ChildRowFeatures, EPropertyBagChildRowFeatures::Menu_MetadataSpecifiers))
 				{
-					MenuBuilder.AddSubMenu(
-						FText::FromName(Metadata::Specifiers::ClampMinName),
-						LOCTEXT("DropDownMenuSubMenuClampMinTooltip", "Set the min value for this property."),
-						BuildGenericSpecifierSubMenu(Metadata::Specifiers::ClampMinName, /*bIsNumeric=*/true),
-						/*bInOpenSubMenuOnClick=*/true);
+					MenuBuilder.BeginSection(/*InExtensionHook=*/NAME_None, LOCTEXT("DropDownMenuSectionMetadataSpecifiers", "Specifiers"));
 
-					MenuBuilder.AddSubMenu(
-						FText::FromName(Metadata::Specifiers::ClampMaxName),
-						LOCTEXT("DropDownMenuSubMenuClampMaxTooltip", "Set the max value for this property."),
-						BuildGenericSpecifierSubMenu(Metadata::Specifiers::ClampMaxName, /*bIsNumeric=*/true),
-						/*bInOpenSubMenuOnClick=*/true);
+					auto BuildGenericSpecifierSubMenu = [this, ChildPropertyHandle, BuildMetadataSpecifierSubMenu](const FName& SpecifierName, const bool bIsNumeric = false)
+					{
+						return BuildMetadataSpecifierSubMenu(
+							TAttribute<FText>::CreateLambda([SpecifierName, ChildPropertyHandle]()
+							{
+								return ChildPropertyHandle.IsValid()
+										   ? FText::FromString(ChildPropertyHandle->GetMetaData(SpecifierName))
+										   : FText::GetEmpty();
+							}),
+							!bIsNumeric
+								? FOnVerifyTextChanged()
+								: FOnVerifyTextChanged::CreateLambda([](const FText& InText, FText& OutMessage) -> bool
+								{
+									if (!InText.IsEmpty() && !InText.IsNumeric())
+									{
+										OutMessage = LOCTEXT("DropDownMenuSpecifierNumericValueExpected", "Value is not numeric");
+										return false;
+									}
+									else
+									{
+										return true;
+									}
+								}),
+							FOnTextCommitted::CreateLambda(
+								[SpecifierName, StructProperty = BagStructProperty, PropUtils = PropUtils, ChildPropertyHandle]
+								(const FText& CommittedText, const ETextCommit::Type CommitType)
+								{
+									if (CommitType != ETextCommit::OnCleared)
+									{
+										ApplyChangesToSinglePropertyDesc(
+											FText::Format(LOCTEXT("DropDownMenuOnSetSpecifier", "Set metadata specifier '{0}' value"), FText::FromName(SpecifierName)),
+											ChildPropertyHandle,
+											StructProperty,
+											[SpecifierName, &CommittedText](FPropertyBagPropertyDesc& Desc)
+												{
+													// If empty, remove the specifier
+													if (CommittedText.IsEmpty())
+													{
+														Desc.RemoveMetadata(SpecifierName);
+													}
+													else
+													{
+														Desc.SetMetaData(SpecifierName, CommittedText.ToString());
+													}
+												});
+									}
+								}));
+					};
 
 					MenuBuilder.AddSubMenu(
-						FText::FromName(Metadata::Specifiers::UIMinName),
-						LOCTEXT("DropDownMenuSubMenuUIMinTooltip", "Set the min UI value for this property."),
-						BuildGenericSpecifierSubMenu(Metadata::Specifiers::UIMinName, /*bIsNumeric=*/true),
+						FText::FromName(Metadata::Specifiers::ToolTipName),
+						LOCTEXT("DropDownMenuSubMenuTooltipTooltip", "Set the tooltip for this property."),
+						BuildGenericSpecifierSubMenu(Metadata::Specifiers::ToolTipName),
 						/*bInOpenSubMenuOnClick=*/true);
 
-					MenuBuilder.AddSubMenu(
-						FText::FromName(Metadata::Specifiers::UIMaxName),
-						LOCTEXT("DropDownMenuSubMenuUIMaxTooltip", "Set the max UI value for this property."),
-						BuildGenericSpecifierSubMenu(Metadata::Specifiers::UIMaxName, /*bIsNumeric=*/true),
-						/*bInOpenSubMenuOnClick=*/true);
-				}
+					if (Metadata::Specifiers::SupportsClamping(*PropertyDesc))
+					{
+						MenuBuilder.AddSubMenu(
+							FText::FromName(Metadata::Specifiers::ClampMinName),
+							LOCTEXT("DropDownMenuSubMenuClampMinTooltip", "Set the min value for this property."),
+							BuildGenericSpecifierSubMenu(Metadata::Specifiers::ClampMinName, /*bIsNumeric=*/true),
+							/*bInOpenSubMenuOnClick=*/true);
+
+						MenuBuilder.AddSubMenu(
+							FText::FromName(Metadata::Specifiers::ClampMaxName),
+							LOCTEXT("DropDownMenuSubMenuClampMaxTooltip", "Set the max value for this property."),
+							BuildGenericSpecifierSubMenu(Metadata::Specifiers::ClampMaxName, /*bIsNumeric=*/true),
+							/*bInOpenSubMenuOnClick=*/true);
+
+						MenuBuilder.AddSubMenu(
+							FText::FromName(Metadata::Specifiers::UIMinName),
+							LOCTEXT("DropDownMenuSubMenuUIMinTooltip", "Set the min UI value for this property."),
+							BuildGenericSpecifierSubMenu(Metadata::Specifiers::UIMinName, /*bIsNumeric=*/true),
+							/*bInOpenSubMenuOnClick=*/true);
+
+						MenuBuilder.AddSubMenu(
+							FText::FromName(Metadata::Specifiers::UIMaxName),
+							LOCTEXT("DropDownMenuSubMenuUIMaxTooltip", "Set the max UI value for this property."),
+							BuildGenericSpecifierSubMenu(Metadata::Specifiers::UIMaxName, /*bIsNumeric=*/true),
+							/*bInOpenSubMenuOnClick=*/true);
+					}
 
-				if (Metadata::Specifiers::SupportsSettingUnits(*PropertyDesc))
-				{
-					MenuBuilder.AddSubMenu(
-						FText::FromName(Metadata::Specifiers::UnitsName),
-						LOCTEXT("DropDownMenuSubMenuUnitsTooltip", "Set the units descriptor for this property."),
-						BuildGenericSpecifierSubMenu(Metadata::Specifiers::UnitsName),
-						/*bInOpenSubMenuOnClick=*/true);
-				}
+					if (Metadata::Specifiers::SupportsSettingUnits(*PropertyDesc))
+					{
+						MenuBuilder.AddSubMenu(
+							FText::FromName(Metadata::Specifiers::UnitsName),
+							LOCTEXT("DropDownMenuSubMenuUnitsTooltip", "Set the units descriptor for this property."),
+							BuildGenericSpecifierSubMenu(Metadata::Specifiers::UnitsName),
+							/*bInOpenSubMenuOnClick=*/true);
+					}
 
-				MenuBuilder.EndSection();
+					MenuBuilder.EndSection();
+				}
 			}
 
 			/*** DROP-DOWN ARROW MENU ***/
@@ -1696,88 +1961,92 @@ void FPropertyBagInstanceDataDetails::OnChildRowAdded(IDetailPropertyRow& ChildR
 				];
 		}
 
-		/*** DRAG AND DROP HANDLER ***/
+		/** DRAG AND DROP HANDLER 
+		 *  We only add drag & drop for reordering if there is no hierarchy root. */
 		if (EnumHasAnyFlags(ChildRowFeatures, EPropertyBagChildRowFeatures::DragAndDrop) && PropertyDesc != nullptr)
 		{
 			TSharedPtr<FPropertyBagDetailsDragDropHandler> DragDropHandler = MakeShared<FPropertyBagDetailsDragDropHandler>(*PropertyDesc);
 
-			// Can accept drag and drop check for if this is a valid drop
-			DragDropHandler->BindCanAcceptDragDrop(
-				FCanAcceptPropertyBagDetailsRowDropOp::CreateLambda(
-					[PropertyDesc](const TSharedPtr<FPropertyBagDetailsDragDropOp>& DropOp, EItemDropZone DropZone) -> TOptional<EItemDropZone>
-					{
-						if (!PropertyDesc || !DropOp.IsValid())
-						{
-							DropOp->SetDecoration(EPropertyBagDropState::Invalid);
-							return TOptional<EItemDropZone>();
-						}
-
-						if (DropZone == EItemDropZone::OntoItem && PropertyDesc->ID != DropOp->PropertyDesc.ID)
+			// If we support drag & drop but a hierarchy root exists, we allow dragging, but not dropping on our property rows
+			if (!HierarchyData.HierarchyRoot.IsValid())
+			{
+				// Can accept drag and drop check for if this is a valid drop
+				DragDropHandler->BindCanAcceptDragDrop(
+					FCanAcceptPropertyBagDetailsRowDropOp::CreateLambda(
+						[PropertyDesc](const TSharedPtr<FPropertyBagDetailsDragDropOp>& DropOp, EItemDropZone DropZone) -> TOptional<EItemDropZone>
 						{
-							DropOp->SetDecoration(EPropertyBagDropState::Invalid);
-							return TOptional<EItemDropZone>();
-						}
+							if (!PropertyDesc || !DropOp.IsValid())
+							{
+								DropOp->SetDecoration(EPropertyBagDropState::Invalid);
+								return TOptional<EItemDropZone>();
+							}
 
-						// No effect to drop in these cases. Either source == target, or moving source above/below target puts source in same location.
-						if (*PropertyDesc == DropOp->PropertyDesc
-							|| (DropZone == EItemDropZone::AboveItem && DropOp->PropertyDesc.GetCachedIndex() == PropertyDesc->GetCachedIndex() - 1)
-							|| (DropZone == EItemDropZone::BelowItem && DropOp->PropertyDesc.GetCachedIndex() == PropertyDesc->GetCachedIndex() + 1))
-						{
-							DropOp->SetDecoration(EPropertyBagDropState::SourceIsTarget);
-							return TOptional<EItemDropZone>();
-						}
+							if (DropZone == EItemDropZone::OntoItem && PropertyDesc->ID != DropOp->PropertyDesc.ID)
+							{
+								DropOp->SetDecoration(EPropertyBagDropState::Invalid);
+								return TOptional<EItemDropZone>();
+							}
 
-						DropOp->SetDecoration(EPropertyBagDropState::Valid);
-						return DropZone;
-					}));
+							// No effect to drop in these cases. Either source == target, or moving source above/below target puts source in same location.
+							if (*PropertyDesc == DropOp->PropertyDesc
+								|| (DropZone == EItemDropZone::AboveItem && DropOp->PropertyDesc.GetCachedIndex() == PropertyDesc->GetCachedIndex() - 1)
+								|| (DropZone == EItemDropZone::BelowItem && DropOp->PropertyDesc.GetCachedIndex() == PropertyDesc->GetCachedIndex() + 1))
+							{
+								DropOp->SetDecoration(EPropertyBagDropState::SourceIsTarget);
+								return TOptional<EItemDropZone>();
+							}
 
-			DragDropHandler->BindOnHandleDragDrop(
-				FOnPropertyBagDetailsRowDropOp::CreateLambda(
-					[WeakSelf, PropertyDesc = *PropertyDesc, BagStructProperty = BagStructProperty, PropUtils = PropUtils](const FPropertyBagPropertyDesc& DroppedPropertyDesc, EItemDropZone DropZone) -> FReply
-					{
-						using namespace UE::StructUtils;
+							DropOp->SetDecoration(EPropertyBagDropState::Valid);
+							return DropZone;
+						}));
 
-						const TSharedPtr<FPropertyBagInstanceDataDetails> DetailsSP = WeakSelf.Pin();
-						const UPropertyBag* ChildBagStruct = DetailsSP ? GetCommonBagStruct(DetailsSP->BagStructProperty) : nullptr;
-						// Validate these properties are still part of the bag.
-						if (!ChildBagStruct
-							|| !ChildBagStruct->FindPropertyDescByProperty(PropertyDesc.CachedProperty)
-							|| !ChildBagStruct->FindPropertyDescByProperty(DroppedPropertyDesc.CachedProperty))
+				DragDropHandler->BindOnHandleDragDrop(
+					FOnPropertyBagDetailsRowDropOp::CreateLambda(
+						[WeakSelf, PropertyDesc = *PropertyDesc, PropertyBagHandle = BagStructProperty, PropUtils = PropUtils](const FPropertyBagPropertyDesc& DroppedPropertyDesc, EItemDropZone DropZone) -> FReply
 						{
-							return FReply::Unhandled();
-						}
+							using namespace UE::StructUtils;
+
+							const TSharedPtr<FPropertyBagInstanceDataDetails> DetailsSP = WeakSelf.Pin();
+							const UPropertyBag* ChildBagStruct = DetailsSP ? GetCommonBagStruct(DetailsSP->BagStructProperty) : nullptr;
+							// Validate these properties are still part of the bag.
+							if (!ChildBagStruct
+								|| !ChildBagStruct->FindPropertyDescByProperty(PropertyDesc.CachedProperty)
+								|| !ChildBagStruct->FindPropertyDescByProperty(DroppedPropertyDesc.CachedProperty))
+							{
+								return FReply::Unhandled();
+							}
 
-						EPropertyBagAlterationResult Result = EPropertyBagAlterationResult::InternalError;
+							EPropertyBagAlterationResult Result = EPropertyBagAlterationResult::InternalError;
 
-						DetailsSP->BagStructProperty->EnumerateRawData([DropZone, &PropertyDesc, &DroppedPropertyDesc, &Result](void* RawData, const int32 /*DataIndex*/, const int32 /*NumData*/)
-						{
-							if (FInstancedPropertyBag* PropertyBag = RawData ? static_cast<FInstancedPropertyBag*>(RawData) : nullptr)
+							DetailsSP->BagStructProperty->EnumerateRawData([DropZone, &PropertyDesc, &DroppedPropertyDesc, &Result](void* RawData, const int32 /*DataIndex*/, const int32 /*NumData*/)
 							{
-								Result = PropertyBag->ReorderProperty(DroppedPropertyDesc.Name, PropertyDesc.Name, DropZone == EItemDropZone::AboveItem);
-							}
+								if (FInstancedPropertyBag* PropertyBag = RawData ? static_cast<FInstancedPropertyBag*>(RawData) : nullptr)
+								{
+									Result = PropertyBag->ReorderProperty(DroppedPropertyDesc.Name, PropertyDesc.Name, DropZone == EItemDropZone::AboveItem);
+								}
 
-							return true;
-						});
+								return true;
+							});
 
-						if (Result == EPropertyBagAlterationResult::Success)
-						{
-							ApplyChangesToSinglePropertyDesc(
-								LOCTEXT("DragDropReorderProperties", "Reordered properties"),
-								DroppedPropertyDesc,
-								BagStructProperty,
-								PropUtils,
-								[PropertyDesc](FPropertyBagPropertyDesc& Desc)
-								{
-									Metadata::SetCategory(Desc, Metadata::GetCategory(PropertyDesc));
-								});
+							if (Result == EPropertyBagAlterationResult::Success)
+							{
+								ApplyChangesToSinglePropertyDesc(
+									LOCTEXT("DragDropReorderProperties", "Reordered properties"),
+									DroppedPropertyDesc,
+									PropertyBagHandle,
+									[PropertyDesc](FPropertyBagPropertyDesc& Desc)
+										{
+											Metadata::SetCategory(Desc, Metadata::GetCategory(PropertyDesc));
+										});
 
-							return FReply::Handled();
-						}
-						else
-						{
-							return FReply::Unhandled();
-						}
-					}));
+								return FReply::Handled();
+							}
+							else
+							{
+								return FReply::Unhandled();
+							}
+						}));
+			}
 
 			// Bind the drag and drop handler for receiving.
 			ChildRow.DragDropHandler(DragDropHandler);
@@ -1813,8 +2082,9 @@ void FPropertyBagInstanceDataDetails::OnChildRowAdded(IDetailPropertyRow& ChildR
 					]
 				];
 		}
-		else // Update the name widget with our new property details composition.
+		else
 		{
+			// Update the name widget with our new property details composition.
 			NameWidget = PropertyDetailsWidget;
 		}
 	}
@@ -1999,34 +2269,48 @@ void FPropertyBagInstanceDataDetails::ResetToDefault(TSharedPtr<IPropertyHandle>
 	ChildPropertyHandle->NotifyFinishedChangingProperties();
 }
 
+EVisibility FPropertyBagInstanceDataDetails::ShouldShowHierarchyProperty(TSharedPtr<IPropertyHandle> ChildPropertyHandle) const
+{
+	if (bShouldShowAdvanced)
+	{
+		return EVisibility::Visible;
+	}
+	
+	if (FProperty* MetaDataProperty = ChildPropertyHandle->GetMetaDataProperty())
+	{
+		return MetaDataProperty->HasAnyPropertyFlags(CPF_AdvancedDisplay) ? EVisibility::Collapsed : EVisibility::Visible;
+	}
+	
+	return EVisibility::Visible;
+}
+
 TSharedRef<SWidget> FPropertyBagInstanceDataDetails::OnPropertyNameContent(TSharedPtr<IPropertyHandle> ChildPropertyHandle, TSharedPtr<SInlineEditableTextBlock> InlineWidget) const
 {
 	constexpr bool bInShouldCloseWindowAfterMenuSelection = true;
 	FMenuBuilder MenuBuilder(bInShouldCloseWindowAfterMenuSelection, nullptr);
 
-	auto MoveProperty = [BagStructProperty = BagStructProperty, PropUtils = PropUtils, ChildPropertyHandle](const int32 Delta)
+	auto MoveProperty = [PropertyBagHandle = BagStructProperty, PropUtils = PropUtils, ChildPropertyHandle](const int32 Delta)
 	{
-		if (!BagStructProperty || !BagStructProperty->IsValidHandle() || !ChildPropertyHandle || !ChildPropertyHandle->IsValidHandle())
+		if (!PropertyBagHandle || !PropertyBagHandle->IsValidHandle() || !ChildPropertyHandle || !ChildPropertyHandle->IsValidHandle())
 		{
 			return;
 		}
 
 		UE::StructUtils::ApplyChangesToPropertyDescs(
-		LOCTEXT("OnPropertyMoved", "Move Property"), BagStructProperty, PropUtils,
-		[&ChildPropertyHandle, &Delta](TArray<FPropertyBagPropertyDesc>& PropertyDescs)
-		{
-			// Move
-			if (PropertyDescs.Num() > 1)
-			{
-				const FProperty* Property = ChildPropertyHandle ? ChildPropertyHandle->GetProperty() : nullptr;
-				const int32 PropertyIndex = PropertyDescs.IndexOfByPredicate([Property](const FPropertyBagPropertyDesc& Desc){ return Desc.CachedProperty == Property; });
-				if (PropertyIndex != INDEX_NONE)
+			LOCTEXT("OnPropertyMoved", "Move Property"), PropertyBagHandle, [&ChildPropertyHandle, &Delta](TArray<FPropertyBagPropertyDesc>& PropertyDescs)
 				{
-					const int32 NewPropertyIndex = FMath::Clamp(PropertyIndex + Delta, 0, PropertyDescs.Num() - 1);
-					PropertyDescs.Swap(PropertyIndex, NewPropertyIndex);
-				}
-			}
-		});
+					// Move
+					if (PropertyDescs.Num() > 1)
+					{
+						const FProperty* Property = ChildPropertyHandle ? ChildPropertyHandle->GetProperty() : nullptr;
+						const int32 PropertyIndex = PropertyDescs.IndexOfByPredicate([Property](const FPropertyBagPropertyDesc& Desc){ return Desc.CachedProperty == Property; });
+						if (PropertyIndex != INDEX_NONE)
+						{
+							const int32 NewPropertyIndex = FMath::Clamp(PropertyIndex + Delta, 0, PropertyDescs.Num() - 1);
+							PropertyDescs.Swap(PropertyIndex, NewPropertyIndex);
+						}
+					}
+				});
 	};
 
 	MenuBuilder.AddWidget(
@@ -2039,7 +2323,7 @@ TSharedRef<SWidget> FPropertyBagInstanceDataDetails::OnPropertyNameContent(TShar
 				PropUtils,
 				SPinTypeSelector::ESelectorType::Full,
 				bAllowContainers,
-				PropertyBagSchemaClass)
+				PropertyBagSchema->GetClass())
 		],
 		FText::GetEmpty());
 
@@ -2081,6 +2365,259 @@ TSharedRef<SWidget> FPropertyBagInstanceDataDetails::OnPropertyNameContent(TShar
 	return MenuBuilder.MakeWidget();
 }
 
+void FPropertyBagInstanceDataDetails::GenerateHierarchySectionRow(IDetailChildrenBuilder& ChildrenBuilder)
+{
+	if (!HierarchyData.HierarchyRoot.IsValid())
+	{
+		return;
+	}
+	
+	TArray<const UHierarchySection*> OrderedSections;
+	TMap<const UHierarchySection*, TArray<const UHierarchyCategory*>> SectionCategoryMap;
+
+	UHierarchyRoot* Root = HierarchyData.HierarchyRoot.Get();
+	for(const UHierarchyElement* Child : Root->GetChildren())
+	{
+		if(const UHierarchyCategory* HierarchyCategory = Cast<UHierarchyCategory>(Child))
+		{
+			const FDataHierarchyElementMetaData_SectionAssociation* SectionAssociation = HierarchyCategory->FindMetaDataOfType<FDataHierarchyElementMetaData_SectionAssociation>();
+			if(SectionAssociation && SectionAssociation->Section.IsValid())
+			{
+				SectionCategoryMap.FindOrAdd(SectionAssociation->Section.Get()).Add(HierarchyCategory);	
+			}
+		}
+	}
+
+	// we create a custom row for user param sections here, but only if we have at least one section specified
+	if(SectionCategoryMap.Num() > 0)
+	{
+		TSharedRef<SWrapBox> SectionsBox = SNew(SWrapBox).UseAllottedSize(true);
+
+		// maps don't guarantee order, so we iterate over the original section data instead
+		for(const TObjectPtr<UHierarchySection>& Section : Root->GetSectionData())
+		{
+			// if the map doesn't contain at the entry, it means there are 0 categories for that section, so we skip it
+			if(!SectionCategoryMap.Contains(Section))
+			{
+				continue;
+			}
+			
+			TArray<const UHierarchyCategory*>& HierarchyCategories = SectionCategoryMap[Section];
+			
+			if(HierarchyCategories.Num() > 0)
+			{
+				// we only want to add a section if any of its categories contain at least one parameter
+				bool bDoesUserParamExist = false;
+				for(const UHierarchyCategory* Category : HierarchyCategories)
+				{
+					bDoesUserParamExist |= Category->DoesOneChildExist<UPropertyBagHierarchyProperty>(true);
+
+					if(bDoesUserParamExist)
+					{
+						break;
+					}
+				}
+
+				if(bDoesUserParamExist)
+				{
+					SectionsBox->AddSlot()
+					.Padding(FMargin(2.f))
+					[
+						SNew(SBox)
+						.Padding(FMargin(0.f, 4.f, 0.f, 0.f))
+						[
+							SNew(SCheckBox)
+							.Style(FAppStyle::Get(), "DetailsView.SectionButton")
+							.OnCheckStateChanged_Lambda([this, Section](ECheckBoxState NewState)
+							{
+								HierarchyData.ActiveHierarchySection = Section;
+								OnRebuildChildren.ExecuteIfBound();
+							})
+							.IsChecked_Lambda([this, Section]()
+							{
+								return HierarchyData.ActiveHierarchySection == Section ? ECheckBoxState::Checked : ECheckBoxState::Unchecked;
+							})
+							[
+								SNew(STextBlock)
+								.TextStyle(FAppStyle::Get(), "SmallText")
+								.Text_UObject(Section.Get(), &UHierarchySection::GetSectionNameAsText)
+								.ToolTipText_UObject(Section.Get(), &UHierarchySection::GetTooltip)
+							]
+						]
+					];
+				}
+			}
+		}
+
+		// if we have at least one custom section, we add a default "All" section
+		if(SectionsBox->GetChildren()->Num() > 0)
+		{
+			SectionsBox->AddSlot()
+			.Padding(FMargin(2.f))
+			[
+				SNew(SBox)
+				.Padding(FMargin(0.f, 4.f, 0.f, 0.f))
+				[
+					SNew(SCheckBox)
+					.Style(FAppStyle::Get(), "DetailsView.SectionButton")
+					.OnCheckStateChanged_Lambda([this](ECheckBoxState NewState)
+					{
+						HierarchyData.ActiveHierarchySection = nullptr;
+						OnRebuildChildren.ExecuteIfBound();
+					})
+					.IsChecked_Lambda([this]()
+					{
+						return HierarchyData.ActiveHierarchySection == nullptr ? ECheckBoxState::Checked : ECheckBoxState::Unchecked;
+					})
+					[
+						SNew(STextBlock)
+						.TextStyle(FAppStyle::Get(), "SmallText")
+						.Text(FText::FromString("All"))
+					]
+				]
+			];
+			
+			FDetailWidgetRow& Row = ChildrenBuilder.AddCustomRow(FText::FromString("Sections"));
+			Row.RowTag(FName("Sections"));
+			
+			Row.WholeRowContent()
+			[
+				SectionsBox
+			];
+		}
+	}
+}
+
+void FPropertyBagInstanceDataDetails::GenerateHierarchyRootRows(IDetailChildrenBuilder& ChildrenBuilder)
+{
+	if (!HierarchyData.HierarchyRoot.IsValid())
+	{
+		return;
+	}
+	
+	// Then, we generate root level categories...
+	TArray<const UPropertyBagHierarchyCategory*> Categories;
+	HierarchyData.HierarchyRoot->GetChildrenOfType(Categories, false);
+		
+	for(const UPropertyBagHierarchyCategory* Category : Categories)
+	{
+		const UHierarchySection* Section = nullptr;
+		if(const FDataHierarchyElementMetaData_SectionAssociation* SectionAssociation = Category->FindMetaDataOfType<FDataHierarchyElementMetaData_SectionAssociation>())
+		{
+			if(SectionAssociation->Section.IsValid())
+			{
+				Section = SectionAssociation->Section.Get();
+			}
+		}
+			
+		// ... but skip categories where the section doesn't match or categories that don't actually contain a property
+		if((HierarchyData.ActiveHierarchySection == nullptr || Section == HierarchyData.ActiveHierarchySection) && Category->DoesOneChildExist<UPropertyBagHierarchyProperty>(true))
+		{
+			ChildrenBuilder.AddCustomBuilder(MakeShared<FHierarchyPropertyCategoryBuilder>(*Category, SharedThis(this)));
+		}
+	}
+		
+	// Then we add root level properties
+	TArray<const UPropertyBagHierarchyProperty*> HierarchyPropertiesUnderRoot;
+	HierarchyData.HierarchyRoot->GetChildrenOfType(HierarchyPropertiesUnderRoot, false);
+	for(const UPropertyBagHierarchyProperty* HierarchyProperty : HierarchyPropertiesUnderRoot)
+	{
+		// we only want to display parameters within the root in the all section
+		if(HierarchyData.ActiveHierarchySection == nullptr)
+		{
+			if (HierarchyData.HierarchyPropertyHandleMap.Contains(HierarchyProperty))
+			{
+				IDetailPropertyRow& PropertyRow = ChildrenBuilder.AddProperty(HierarchyData.HierarchyPropertyHandleMap[HierarchyProperty].ToSharedRef());
+				PropertyRow.Visibility(TAttribute<EVisibility>::CreateSP(this, &FPropertyBagInstanceDataDetails::ShouldShowHierarchyProperty, PropertyRow.GetPropertyHandle()));
+				OnChildRowAdded(PropertyRow);
+			}
+		}
+	}
+}
+
+void FPropertyBagInstanceDataDetails::GenerateHierarchyLeftoverPropertyRows(IDetailChildrenBuilder& ChildrenBuilder, const TArray<TSharedPtr<IPropertyHandle>>& AllChildPropertyHandles)
+{
+	if (!HierarchyData.HierarchyRoot.IsValid())
+	{
+		return;
+	}
+	
+	// Then we add all parameters that haven't been setup in the hierarchy at all, if the active section is set to "All"
+	if (HierarchyData.ActiveHierarchySection == nullptr)
+	{
+		TArray<const UPropertyBagHierarchyProperty*> AllHierarchyProperties;
+		HierarchyData.HierarchyRoot->GetChildrenOfType(AllHierarchyProperties, true);
+		
+		TArray<TSharedPtr<IPropertyHandle>> LeftoverProperties = AllChildPropertyHandles;
+		TArray<TSharedPtr<IPropertyHandle>> PropertyHandlesInHierarchy;
+		HierarchyData.HierarchyPropertyHandleMap.GenerateValueArray(PropertyHandlesInHierarchy);
+			
+		LeftoverProperties.RemoveAll([PropertyHandlesInHierarchy](TSharedPtr<IPropertyHandle> RemovalCandidate)
+		{
+			return PropertyHandlesInHierarchy.Contains(RemovalCandidate);
+		});
+			
+		for (TSharedPtr<IPropertyHandle> LeftoverProperty : LeftoverProperties)
+		{
+			IDetailPropertyRow& PropertyRow = ChildrenBuilder.AddProperty(LeftoverProperty.ToSharedRef());
+			PropertyRow.Visibility(TAttribute<EVisibility>::CreateSP(this, &FPropertyBagInstanceDataDetails::ShouldShowHierarchyProperty, PropertyRow.GetPropertyHandle()));
+			OnChildRowAdded(PropertyRow);
+		}
+	}
+}
+
+const FSlateBrush* FPropertyBagInstanceDataDetails::GetAdvancedImage() const
+{	
+	FName ResourceName;
+	if (bShouldShowAdvanced)
+	{
+		if (ShowAdvancedButton->IsHovered())
+		{
+			static const FName ExpandedHoveredName = "DetailsView.PulldownArrow.Up.Hovered";
+			ResourceName = ExpandedHoveredName;
+		}
+		else
+		{
+			static const FName ExpandedName = "DetailsView.PulldownArrow.Up";
+			ResourceName = ExpandedName;
+		}
+	}
+	else
+	{
+		if (ShowAdvancedButton->IsHovered())
+		{
+			static const FName CollapsedHoveredName = "DetailsView.PulldownArrow.Down.Hovered";
+			ResourceName = CollapsedHoveredName;
+		}
+		else
+		{
+			static const FName CollapsedName = "DetailsView.PulldownArrow.Down";
+			ResourceName = CollapsedName;
+		}
+	}
+	
+	return FAppStyle::GetBrush(ResourceName);
+}
+
+EVisibility FPropertyBagInstanceDataDetails::GetAdvancedButtonVisibility() const
+{
+	if (HierarchyData.HierarchyRoot.IsValid())
+	{
+		TArray<UPropertyBagHierarchyProperty*> AllHierarchyProperties;
+		HierarchyData.HierarchyRoot->GetChildrenOfType<UPropertyBagHierarchyProperty>(AllHierarchyProperties, true);
+		
+		for (UPropertyBagHierarchyProperty* Property : AllHierarchyProperties)
+		{
+			if (Property->bAdvanced)
+			{
+				return EVisibility::Visible;
+			}
+		}
+	}
+	
+	return EVisibility::Collapsed;
+}
+
 ////////////////////////////////////
 
 TSharedRef<IPropertyTypeCustomization> FPropertyBagDetails::MakeInstance()
@@ -2095,6 +2632,8 @@ void FPropertyBagDetails::CustomizeHeader(TSharedRef<IPropertyHandle> StructProp
 	StructProperty = StructPropertyHandle;
 	check(StructProperty);
 
+	PropertyBagSchemaCDO = UE::StructUtils::ExtractPropertyBagSchemaCDO(StructPropertyHandle);
+	
 	if (const FProperty* MetaDataProperty = StructProperty->GetMetaDataProperty())
 	{
 		PRAGMA_DISABLE_DEPRECATION_WARNINGS
@@ -2104,9 +2643,7 @@ void FPropertyBagDetails::CustomizeHeader(TSharedRef<IPropertyHandle> StructProp
 		bAllowContainers = MetaDataProperty->HasMetaData(UE::StructUtils::Metadata::AllowContainersName)
 			? MetaDataProperty->GetBoolMetaData(UE::StructUtils::Metadata::AllowContainersName)
 			: true;
-
-		PropertyBagSchemaClass = FindObject<UClass>(nullptr, MetaDataProperty->GetMetaData(UE::StructUtils::Metadata::PropertyBagSchemaClassName));
-
+		
 		if (MetaDataProperty->HasMetaData(UE::StructUtils::Metadata::DefaultTypeName))
 		{
 			if (UEnum* Enum = StaticEnum<EPropertyBagPropertyType>())
@@ -2146,13 +2683,12 @@ void FPropertyBagDetails::CustomizeHeader(TSharedRef<IPropertyHandle> StructProp
 		}
 	}
 
-	TSharedPtr<SWidget> ValueWidget = SNullWidget::NullWidget;
+	TSharedPtr<SHorizontalBox> ValueWidget = SNew(SHorizontalBox);
 	PRAGMA_DISABLE_DEPRECATION_WARNINGS
 	if (!bFixedLayout && ChildRowFeatures != EPropertyBagChildRowFeatures::Fixed)
 	PRAGMA_ENABLE_DEPRECATION_WARNINGS
 	{
-		SAssignNew(ValueWidget, SHorizontalBox)
-		+ SHorizontalBox::Slot()
+		ValueWidget->AddSlot()
 		.AutoWidth()
 		[
 			MakeAddPropertyWidget(StructProperty, PropUtils, DefaultType).ToSharedRef()
@@ -2174,6 +2710,8 @@ void FPropertyBagDetails::CustomizeHeader(TSharedRef<IPropertyHandle> StructProp
 
 void FPropertyBagDetails::CustomizeChildren(TSharedRef<IPropertyHandle> StructPropertyHandle, IDetailChildrenBuilder& StructBuilder, IPropertyTypeCustomizationUtils& StructCustomizationUtils)
 {
+	FPropertyBagInstanceDataDetails::FConstructParams::FHierarchyParams HierarchyParams;
+	HierarchyParams.OnNavigateToHierarchyPropertyRequested = FPropertyBagInstanceDataDetails::FPropertyBagPropertyDelegate::CreateSP(this, &FPropertyBagDetails::NavigateToHierarchyProperty); 
 	FPropertyBagInstanceDataDetails::FConstructParams Params
 	{
 		.BagStructProperty = StructProperty,
@@ -2182,7 +2720,7 @@ void FPropertyBagDetails::CustomizeChildren(TSharedRef<IPropertyHandle> StructPr
 		PRAGMA_DISABLE_DEPRECATION_WARNINGS
 		.ChildRowFeatures = bFixedLayout ? EPropertyBagChildRowFeatures::Fixed : ChildRowFeatures,
 		PRAGMA_ENABLE_DEPRECATION_WARNINGS
-		.PropertyBagSchemaClass = PropertyBagSchemaClass
+		.HierarchyParams = HierarchyParams
 	};
 
 	const UPropertyBagDetailsExtension* Extension = GetDefault<UPropertyBagDetailsExtension>();
@@ -2205,79 +2743,340 @@ void FPropertyBagDetails::CustomizeChildren(TSharedRef<IPropertyHandle> StructPr
 	StructBuilder.AddCustomBuilder(InstanceDetails);
 }
 
+FReply FPropertyBagDetails::SummonHierarchyEditor(TSharedPtr<IPropertyHandle> InPropertyBagHandle, TSharedPtr<IPropertyUtilities> InPropertyUtilities)
+{
+	if (!InPropertyBagHandle.IsValid() || !InPropertyBagHandle->IsValidHandle() || !InPropertyUtilities.IsValid())
+	{
+		return FReply::Unhandled();
+	}
+	
+	UPropertyBagHierarchyRoot* HierarchyRoot = UE::StructUtils::ExtractHierarchyRoot(InPropertyBagHandle.ToSharedRef());
+	
+	if (!HierarchyRoot)
+	{
+		return FReply::Unhandled();
+	}
+	
+	bool bNewWindow = false;
+	
+	TSharedPtr<SWindow> Window;
+
+	TArray<UObject*> OuterObjects;
+	InPropertyBagHandle->GetOuterObjects(OuterObjects);
+	
+	if (OuterObjects.Num() != 1)
+	{
+		return FReply::Unhandled();
+	}
+	
+	FString OwnerName = OuterObjects[0]->GetName();
+	
+	FText Title = FText::FormatOrdered(INVTEXT("{0} - {1}"), FText::FromString(OwnerName), InPropertyBagHandle->GetPropertyDisplayName());
+	// We reuse the existing window if possible
+	if (!PropertyBagHierarchyEditorWindow.IsValid())
+	{
+		Window = SNew(SWindow)
+		.Title(Title)
+		.IsTopmostWindow(true)
+		.ClientSize(FVector2f(750.f, 450.f));
+		
+		PropertyBagHierarchyEditorWindow = Window;
+		bNewWindow = true;
+	}
+	else
+	{
+		Window = PropertyBagHierarchyEditorWindow.Pin();
+	}
+	
+	Window->SetTitle(Title);
+	
+	TSharedPtr<SPropertyBagHierarchyEditor> NewPropertyBagHierarchyEditor = SNew(SPropertyBagHierarchyEditor, InPropertyBagHandle.ToSharedRef(), InPropertyUtilities.ToSharedRef())
+	.OnCloseRequested_Static(&FPropertyBagDetails::ClosePropertyBagHierarchyEditorWindow);
+	
+	PropertyBagHierarchyEditor = NewPropertyBagHierarchyEditor;
+	Window->SetContent(NewPropertyBagHierarchyEditor.ToSharedRef());
+	
+	if (bNewWindow)
+	{
+		FSlateApplication::Get().AddWindow(Window.ToSharedRef(), true);
+	}
+	else
+	{
+		FWindowDrawAttentionParameters Parameters;
+		Parameters.RequestType = EWindowDrawAttentionRequestType::Stop;
+		Window->DrawAttention(Parameters);
+	}
+	
+	return FReply::Handled();
+}
+
+void FPropertyBagDetails::NavigateToHierarchyProperty(const FPropertyBagPropertyDesc& Desc)
+{
+	if (!Desc.ID.IsValid())
+	{
+		return;
+	}
+	
+	if (!PropertyBagHierarchyEditor.IsValid() || !PropertyBagHierarchyEditor.Pin()->GetPropertyBagHandle().IsValid() || !PropertyBagHierarchyEditor.Pin()->GetPropertyBagHandle()->IsSamePropertyNode(StructProperty))
+	{
+		SummonHierarchyEditor(StructProperty.ToSharedRef(), PropUtils.ToSharedRef());
+	}
+	
+	if (PropertyBagHierarchyEditor.IsValid())
+	{
+		PropertyBagHierarchyEditor.Pin()->NavigateToProperty(Desc, true);
+	}
+}
+
+void FPropertyBagDetails::ClosePropertyBagHierarchyEditorWindow()
+{
+	if (PropertyBagHierarchyEditorWindow.IsValid())
+	{
+		PropertyBagHierarchyEditorWindow.Pin()->RequestDestroyWindow();
+	}
+}
+
 TSharedPtr<SWidget> FPropertyBagDetails::MakeAddPropertyWidget(TSharedPtr<IPropertyHandle> InStructProperty, TSharedPtr<IPropertyUtilities> InPropUtils, EPropertyBagPropertyType DefaultType, const FSlateColor IconColor)
 {
 	const FText RemoveAllPropertiesText = LOCTEXT("OnRemoveAllProperties", "Remove all Properties");
 
-	return SNew(SHorizontalBox)
-		+SHorizontalBox::Slot()
-		.AutoWidth()
-		[
-			SNew(SButton)
-			.VAlign(VAlign_Center)
-			.HAlign(HAlign_Center)
-			.ButtonStyle(FAppStyle::Get(), "SimpleButton")
-			.ToolTipText(LOCTEXT("AddProperty_Tooltip", "Add new property"))
-			.OnClicked_Lambda([StructProperty = InStructProperty, PropUtils = InPropUtils, DefaultType]()
+	TSharedRef<SHorizontalBox> Container = SNew(SHorizontalBox)
+	+SHorizontalBox::Slot()
+	.AutoWidth()
+	[
+		SNew(SButton)
+		.VAlign(VAlign_Center)
+		.HAlign(HAlign_Center)
+		.ButtonStyle(FAppStyle::Get(), "SimpleButton")
+		.ToolTipText(LOCTEXT("AddProperty_Tooltip", "Add new property"))
+		.OnClicked_Lambda([StructProperty = InStructProperty, PropUtils = InPropUtils, DefaultType]()
+		{
+			constexpr int32 MaxIterations = 100;
+			FName NewName(TEXT("NewProperty"));
+			int32 Number = 1;
+			while (!UE::StructUtils::Private::IsUniqueName(NewName, FName(), StructProperty) && Number < MaxIterations)
 			{
-				constexpr int32 MaxIterations = 100;
-				FName NewName(TEXT("NewProperty"));
-				int32 Number = 1;
-				while (!UE::StructUtils::Private::IsUniqueName(NewName, FName(), StructProperty) && Number < MaxIterations)
-				{
-					Number++;
-					NewName.SetNumber(Number);
-				}
-				if (Number == MaxIterations)
-				{
-					return FReply::Handled();
-				}
+				Number++;
+				NewName.SetNumber(Number);
+			}
+			if (Number == MaxIterations)
+			{
+				return FReply::Handled();
+			}
 
-				UE::StructUtils::ApplyChangesToPropertyDescs(
-					LOCTEXT("OnPropertyAdded", "Add Property"), StructProperty, PropUtils,
-					[&NewName, DefaultType](TArray<FPropertyBagPropertyDesc>& PropertyDescs)
-					{
-						PropertyDescs.Emplace(NewName, DefaultType);
-					});
+			UE::StructUtils::ApplyChangesToPropertyDescs(
+				LOCTEXT("OnPropertyAdded", "Add Property"), StructProperty,
+				[&NewName, DefaultType](TArray<FPropertyBagPropertyDesc>& PropertyDescs)
+				{
+					PropertyDescs.Emplace(NewName, DefaultType);
+				});
 
-				return FReply::Handled();
+			return FReply::Handled();
 
-			})
-			.Content()
+		})
+		.Content()
+		[
+			SNew(SImage)
+			.Image(FAppStyle::GetBrush("Icons.PlusCircle"))
+			.ColorAndOpacity(IconColor)
+		]
+	]
+	+ SHorizontalBox::Slot()
+	.AutoWidth()
+	.Padding(FMargin(0.f, 0.f, 4.f, 0.f))
+	[
+		PropertyCustomizationHelpers::MakeEmptyButton(
+			FSimpleDelegate::CreateLambda([StructProperty = InStructProperty, PropUtils = InPropUtils, RemoveAllPropertiesText]()
+				{
+					UE::StructUtils::ApplyChangesToPropertyDescs(
+						RemoveAllPropertiesText,
+						StructProperty,
+						[](TArray<FPropertyBagPropertyDesc>& PropertyDescs)
+						{
+							PropertyDescs.Empty();
+						});
+				}),
+			RemoveAllPropertiesText
+			)
+	];
+	
+	// We attempt to retrieve a hierarchy root via schema, and display the Edit Hierarchy button if successful
+	if (UPropertyBagHierarchyRoot* Root = UE::StructUtils::ExtractHierarchyRoot(InStructProperty.ToSharedRef()))
+	{		
+		Container->AddSlot()
+		.AutoWidth()
+		.HAlign(HAlign_Left)
+		.VAlign(VAlign_Center)
+		[
+			SNew(SButton)
+			.ToolTipText(LOCTEXT("EditHierarchy_Tooltip", "Summon Hierarchy Editor to edit categories and metadata."))
+			.ButtonStyle(&FAppStyle::GetWidgetStyle<FButtonStyle>("HoverHintOnly"))
+			.OnClicked_Static(&SummonHierarchyEditor, InStructProperty, InPropUtils)
 			[
 				SNew(SImage)
-				.Image(FAppStyle::GetBrush("Icons.PlusCircle"))
+				.Image(FAppStyle::GetBrush("Icons.EditHierarchy"))
 				.ColorAndOpacity(IconColor)
 			]
-		]
-		+ SHorizontalBox::Slot()
-		.AutoWidth()
-		.Padding(FMargin(0.f, 0.f, 4.f, 0.f))
-		[
-			PropertyCustomizationHelpers::MakeEmptyButton(
-				FSimpleDelegate::CreateLambda([StructProperty = InStructProperty, PropUtils = InPropUtils, RemoveAllPropertiesText]()
-					{
-						UE::StructUtils::ApplyChangesToPropertyDescs(
-							RemoveAllPropertiesText,
-							StructProperty,
-							PropUtils,
-							[](TArray<FPropertyBagPropertyDesc>& PropertyDescs)
-							{
-								PropertyDescs.Empty();
-							});
-					}),
-				RemoveAllPropertiesText
-				)
 		];
-
+	}
+	
+	return Container;
 }
 
 ////////////////////////////////////
-
-bool UPropertyBagSchema::SupportsPinTypeContainer(TWeakPtr<const FEdGraphSchemaAction> SchemaAction,
-		const FEdGraphPinType& PinType, const EPinContainerType& ContainerType) const
+bool UPropertyBagSchema::SupportsPinTypeContainer(TWeakPtr<const FEdGraphSchemaAction> SchemaAction, const FEdGraphPinType& PinType, const EPinContainerType& ContainerType) const
 {
 	return ContainerType == EPinContainerType::None || ContainerType == EPinContainerType::Array || ContainerType == EPinContainerType::Set || ContainerType == EPinContainerType::Map;
 }
 
+UPropertyBagHierarchyRoot* UPropertyBagSchema::GetHierarchyRoot(const TArray<UObject*>& ObjectsWithProperty) const
+{
+	return nullptr;
+}
+
+TSubclassOf<UPropertyBagHierarchyProperty> UPropertyBagSchema::GetHierarchyPropertyType(const FPropertyBagPropertyDesc& Desc) const
+{
+	if (Desc.IsNumericType() && Desc.ValueType != EPropertyBagPropertyType::Bool && Desc.ContainerTypes.IsEmpty())
+	{			
+		return UPropertyBagHierarchyNumericProperty::StaticClass();
+	}
+	
+	return UPropertyBagHierarchyProperty::StaticClass();
+}
+
+TSubclassOf<UPropertyBagHierarchyCategory> UPropertyBagSchema::GetHierarchyCategoryType() const
+{
+	return UPropertyBagHierarchyCategory::StaticClass();
+}
+
+TSubclassOf<UPropertyBagHierarchySection> UPropertyBagSchema::GetHierarchySectionType() const
+{
+	return UPropertyBagHierarchySection::StaticClass();
+}
+
+void UPropertyBagSchema::TransferPropertyBagMetadataIntoHierarchy(const FInstancedPropertyBag& PropertyBag,	UPropertyBagHierarchyRoot& HierarchyRoot) const
+{
+	TMap<const FPropertyBagPropertyDesc*, TArray<FString>> CategoryChains;
+	
+	for (const FPropertyBagPropertyDesc& Desc : PropertyBag.GetPropertyBagStruct()->GetPropertyDescs())
+	{
+		FString CategoryString = Desc.GetMetaData(UE::StructUtils::Metadata::CategoryName);
+		TArray<FString> CategoriesOfDesc;
+		CategoryString.ParseIntoArray(CategoriesOfDesc, TEXT("|"), true);
+		
+		if (CategoriesOfDesc.Num() > 0)
+		{
+			TArray<FString>& Array = CategoryChains.Add(&Desc);
+			
+			for (const FString& Category : CategoriesOfDesc)
+			{
+				Array.Add(Category);
+			}
+		}
+	}
+	
+	auto GenerateHierarchyRecursive = [&](auto&& Self, const FPropertyBagPropertyDesc& Desc, const TArray<FString>& CategoryChain, UHierarchyElement* Outer) -> void
+	{
+		if (CategoryChain.Num() > 0)
+		{				
+			TArray<UPropertyBagHierarchyCategory*> ExistingCategories;
+			Outer->GetChildrenOfType<UPropertyBagHierarchyCategory>(ExistingCategories, false);
+				
+			FString CurrentCategory = CategoryChain[0];
+			
+			UPropertyBagHierarchyCategory* HierarchyCategory = nullptr;
+		
+			UPropertyBagHierarchyCategory** ExistingHierarchyCategory = ExistingCategories.FindByPredicate([CurrentCategory](UPropertyBagHierarchyCategory* Candidate)
+			{
+				return Candidate->GetCategoryName().ToString().Equals(CurrentCategory);
+			});
+		
+			if (ExistingHierarchyCategory)
+			{
+				HierarchyCategory = *ExistingHierarchyCategory;
+			}
+			else
+			{
+				HierarchyCategory = NewObject<UPropertyBagHierarchyCategory>(Outer, GetHierarchyCategoryType());
+				HierarchyCategory->SetCategoryName(FName(CurrentCategory));
+				Outer->GetChildrenMutable().Add(HierarchyCategory);
+			}
+			
+			TArray<FString> ChildCategoryChain = CategoryChain;
+			ChildCategoryChain.RemoveAt(0);
+			Self(Self, Desc, ChildCategoryChain, HierarchyCategory);
+		}
+		else
+		{
+			TSubclassOf<UPropertyBagHierarchyProperty> PropertyClass = GetHierarchyPropertyType(Desc);
+			UPropertyBagHierarchyProperty* HierarchyProperty = NewObject<UPropertyBagHierarchyProperty>(Outer, PropertyClass);
+			HierarchyProperty->Initialize(Desc);
+			Outer->GetChildrenMutable().Add(HierarchyProperty);
+			
+			if (Desc.HasMetaData(UE::StructUtils::Metadata::Specifiers::ToolTipName))
+			{
+				HierarchyProperty->Tooltip = FText::FromString(Desc.GetMetaData(UE::StructUtils::Metadata::Specifiers::ToolTipName));
+			}
+			
+			if (UPropertyBagHierarchyNumericProperty* NumericProperty = Cast<UPropertyBagHierarchyNumericProperty>(HierarchyProperty))
+			{
+				if (Desc.HasMetaData(UE::StructUtils::Metadata::Specifiers::ClampMinName))
+				{
+					float ClampMinValue;
+					if (LexTryParseString(ClampMinValue, *Desc.GetMetaData(UE::StructUtils::Metadata::Specifiers::ClampMinName)))
+					{
+						NumericProperty->bUseClampMin = true;
+						NumericProperty->ClampMin = ClampMinValue;	
+					}
+				}
+				
+				if (Desc.HasMetaData(UE::StructUtils::Metadata::Specifiers::ClampMaxName))
+				{
+					float ClampMaxValue;
+					if (LexTryParseString(ClampMaxValue, *Desc.GetMetaData(UE::StructUtils::Metadata::Specifiers::ClampMaxName)))
+					{
+						NumericProperty->bUseClampMax = true;
+						NumericProperty->ClampMax = ClampMaxValue;	
+					}
+				}
+				
+				if (Desc.HasMetaData(UE::StructUtils::Metadata::Specifiers::UIMinName))
+				{
+					float UIMinValue;
+					if (LexTryParseString(UIMinValue, *Desc.GetMetaData(UE::StructUtils::Metadata::Specifiers::UIMinName)))
+					{
+						NumericProperty->bUseUIMin = true;
+						NumericProperty->UIMin = UIMinValue;	
+					}
+				}
+				
+				if (Desc.HasMetaData(UE::StructUtils::Metadata::Specifiers::UIMaxName))
+				{
+					float UIMaxValue;
+					if (LexTryParseString(UIMaxValue, *Desc.GetMetaData(UE::StructUtils::Metadata::Specifiers::UIMaxName)))
+					{
+						NumericProperty->bUseUIMax = true;
+						NumericProperty->UIMax = UIMaxValue;	
+					}
+				}
+				
+				if (Desc.HasMetaData(UE::StructUtils::Metadata::Specifiers::UnitsName))
+				{
+					FString UnitString = Desc.GetMetaData(UE::StructUtils::Metadata::Specifiers::UnitsName);
+					if (TOptional<EUnit> Unit = FUnitConversion::UnitFromString(*UnitString))
+					{
+						NumericProperty->bUseUnit = true;
+						NumericProperty->Unit = Unit.GetValue();
+					}
+				}
+			}
+		}
+	};
+	
+	for (const auto& CategoryChain : CategoryChains)
+	{
+		GenerateHierarchyRecursive(GenerateHierarchyRecursive, *CategoryChain.Key, CategoryChain.Value, &HierarchyRoot);
+	}
+}
+
 #undef LOCTEXT_NAMESPACE
diff --git a/Engine/Source/Editor/StructUtilsEditor/Private/HierarchyEditor/PropertyBagHierarchyViewModel.cpp b/Engine/Source/Editor/StructUtilsEditor/Private/HierarchyEditor/PropertyBagHierarchyViewModel.cpp
new file mode 100644
--- /dev/null
+++ b/Engine/Source/Editor/StructUtilsEditor/Private/HierarchyEditor/PropertyBagHierarchyViewModel.cpp
@@ -0,0 +1,644 @@
+﻿// Copyright Epic Games, Inc. All Rights Reserved.
+
+#include "HierarchyEditor/PropertyBagHierarchyViewModel.h"
+
+#include "IDetailTreeNode.h"
+#include "IPropertyUtilities.h"
+#include "PropertyBagDetails.h"
+#include "PropertyHandle.h"
+#include "PropertyPath.h"
+#include "StructUtilsEditorUtilsPrivate.h"
+#include "StructUtilsMetadata.h"
+#include "Framework/Notifications/NotificationManager.h"
+#include "Modules/ModuleManager.h"
+#include "Subsystems/AssetEditorSubsystem.h"
+#include "Widgets/Notifications/SNotificationList.h"
+#include "Editor.h"
+
+#define LOCTEXT_NAMESPACE "PropertyBagHierarchyEditor"
+
+void UPropertyBagHierarchyProperty::Initialize(const FPropertyBagPropertyDesc& InPropertyDesc)
+{
+	SetIdentity(ConstructIdentity(InPropertyDesc));
+}
+
+FHierarchyElementIdentity UPropertyBagHierarchyProperty::ConstructIdentity(const FPropertyBagPropertyDesc& PropertyDesc)
+{
+	return FHierarchyElementIdentity({PropertyDesc.ID}, {});
+}
+
+FGuid UPropertyBagHierarchyProperty::GetPropertyId() const
+{
+	if (Identity.Guids.Num() == 1)
+	{
+		return Identity.Guids[0];
+	}
+	
+	return FGuid();
+}
+
+FPropertyBagHierarchyPropertyViewModel::FPropertyBagHierarchyPropertyViewModel(UPropertyBagHierarchyProperty* InProperty, TSharedRef<FHierarchyElementViewModel> InParent, TWeakObjectPtr<UDataHierarchyViewModelBase> InHierarchyViewModel)
+	: FHierarchyItemViewModel(InProperty, InParent, InHierarchyViewModel)
+{
+	
+}
+
+bool FPropertyBagHierarchyPropertyViewModel::DoesExternalDataStillExist(const UHierarchyDataRefreshContext* Context) const
+{
+	// If the property desc doesn't exist at all anymore, we delete it
+	const FPropertyBagPropertyDesc* Desc = GetPropertyDesc();
+	if (Desc == nullptr)
+	{
+		return false;
+	}
+	
+	TSharedPtr<IPropertyHandle> PropertyBagHandle = Cast<UPropertyBagHierarchyViewModel>(HierarchyViewModel.Get())->GetPropertyBagHandle();
+	
+	// If the handle is not valid, we assume the data still exists. Handle might have just become invalidated.
+	if (!PropertyBagHandle.IsValid() || !PropertyBagHandle->IsValidHandle())
+	{
+		return true;
+	}
+	
+	// Alternatively, if the hierarchy type has changed, we also delete it. This can be the case if the user changes the underlying property type.
+	// @TODO (ME) Ideally we would create the new hierarchy element in place, but for now we accept deleting it automatically
+	if (const UPropertyBagSchema* PropertyBagSchemaCDO = UE::StructUtils::ExtractPropertyBagSchemaCDO(PropertyBagHandle.ToSharedRef()))
+	{
+		TSubclassOf<UPropertyBagHierarchyProperty> Type = PropertyBagSchemaCDO->GetHierarchyPropertyType(*Desc);
+		if (Type != GetData()->GetClass())
+		{
+			return false;
+		}
+	}
+	
+	return true;
+}
+
+const FPropertyBagPropertyDesc* FPropertyBagHierarchyPropertyViewModel::GetPropertyDesc() const
+{
+	return Cast<UPropertyBagHierarchyViewModel>(HierarchyViewModel)->GetPropertyDescForProperty(GetData<UPropertyBagHierarchyProperty>()->GetPropertyId());
+}
+
+void FPropertyBagHierarchyPropertyViewModel::PostEditChange(const FPropertyChangedEvent& PropertyChangedEvent, FProperty* PropertyThatChanged)
+{
+	if (PropertyChangedEvent.Property)
+	{
+		if (PropertyChangedEvent.Property->GetName() == GET_MEMBER_NAME_CHECKED(UPropertyBagHierarchyProperty, Tooltip))
+		{
+			UE::StructUtils::ApplyChangesToSinglePropertyDesc(
+				LOCTEXT("ChangedTooltip", "Changed tooltip"),
+				*GetPropertyDesc(),
+				Cast<UPropertyBagHierarchyViewModel>(HierarchyViewModel)->GetPropertyBagHandle(),
+				[Tooltip = GetData<UPropertyBagHierarchyProperty>()->Tooltip](FPropertyBagPropertyDesc& Desc)
+					{
+						// If empty, remove the specifier
+						if (Tooltip.IsEmpty())
+						{
+							Desc.RemoveMetadata(UE::StructUtils::Metadata::Specifiers::ToolTipName);
+						}
+						else
+						{
+							Desc.SetMetaData(UE::StructUtils::Metadata::Specifiers::ToolTipName, Tooltip.ToString());
+						}
+					});
+		}
+		
+		if (PropertyChangedEvent.Property->GetName() == GET_MEMBER_NAME_CHECKED(UPropertyBagHierarchyProperty, bAdvanced))
+		{
+			UE::StructUtils::ApplyChangesToSinglePropertyDesc(
+				LOCTEXT("ToggleAdvanced", "Toggled advance"),
+				*GetPropertyDesc(),
+				Cast<UPropertyBagHierarchyViewModel>(HierarchyViewModel)->GetPropertyBagHandle(),
+				[bInAdvanced = GetData<UPropertyBagHierarchyProperty>()->bAdvanced](FPropertyBagPropertyDesc& Desc)
+					{
+						if (bInAdvanced)
+						{
+							Desc.PropertyFlags |= CPF_AdvancedDisplay;
+						}
+						else
+						{
+							Desc.PropertyFlags &= ~CPF_AdvancedDisplay;
+						}
+					});
+		}
+	}
+	
+	Cast<UPropertyBagHierarchyViewModel>(HierarchyViewModel)->RebuildDetails();
+}
+
+FString FPropertyBagHierarchyPropertyViewModel::ToString() const
+{
+	const FPropertyBagPropertyDesc* Desc = GetPropertyDesc();
+	return Desc ? Desc->Name.ToString() : TEXT("Unidentified");
+}
+
+FPropertyBagHierarchyNumericPropertyViewModel::FPropertyBagHierarchyNumericPropertyViewModel(UPropertyBagHierarchyNumericProperty* InProperty, TSharedRef<FHierarchyElementViewModel> InParent, TWeakObjectPtr<UDataHierarchyViewModelBase> InHierarchyViewModel)
+	: FPropertyBagHierarchyPropertyViewModel(InProperty, InParent, InHierarchyViewModel)
+{
+}
+
+void FPropertyBagHierarchyNumericPropertyViewModel::PostEditChange(const FPropertyChangedEvent& PropertyChangedEvent, FProperty* PropertyThatChanged)
+{
+	if (PropertyChangedEvent.Property)
+	{
+		
+		if (PropertyChangedEvent.Property->GetName() == GET_MEMBER_NAME_CHECKED(UPropertyBagHierarchyNumericProperty, bUseClampMin))
+		{
+			UE::StructUtils::ApplyChangesToSinglePropertyDesc(
+				LOCTEXT("ChangedUseClampMin", "Changed Use Clamp Min"),
+				*GetPropertyDesc(),
+				Cast<UPropertyBagHierarchyViewModel>(HierarchyViewModel)->GetPropertyBagHandle(),
+				[bUseClampMin = GetData<UPropertyBagHierarchyNumericProperty>()->bUseClampMin, ClampMin = GetData<UPropertyBagHierarchyNumericProperty>()->ClampMin](FPropertyBagPropertyDesc& Desc)
+					{
+						if (bUseClampMin)
+						{
+							Desc.SetMetaData(UE::StructUtils::Metadata::Specifiers::ClampMinName, FString::SanitizeFloat(ClampMin));
+						}
+						else
+						{
+							Desc.RemoveMetadata(UE::StructUtils::Metadata::Specifiers::ClampMinName);
+						}
+					});
+		}
+		
+		if (PropertyChangedEvent.Property->GetName() == GET_MEMBER_NAME_CHECKED(UPropertyBagHierarchyNumericProperty, ClampMin))
+		{
+			UE::StructUtils::ApplyChangesToSinglePropertyDesc(
+				LOCTEXT("ChangedClampMin", "Changed Clamp Min"),
+				*GetPropertyDesc(),
+				Cast<UPropertyBagHierarchyViewModel>(HierarchyViewModel)->GetPropertyBagHandle(),
+				[ClampMin = GetData<UPropertyBagHierarchyNumericProperty>()->ClampMin](FPropertyBagPropertyDesc& Desc)
+					{
+						Desc.SetMetaData(UE::StructUtils::Metadata::Specifiers::ClampMinName, FString::SanitizeFloat(ClampMin));
+					});
+		}
+		
+		if (PropertyChangedEvent.Property->GetName() == GET_MEMBER_NAME_CHECKED(UPropertyBagHierarchyNumericProperty, bUseClampMax))
+		{
+			UE::StructUtils::ApplyChangesToSinglePropertyDesc(
+				LOCTEXT("ChangedUseClampMax", "Changed Use Clamp Max"),
+				*GetPropertyDesc(),
+				Cast<UPropertyBagHierarchyViewModel>(HierarchyViewModel)->GetPropertyBagHandle(),
+				[bUseClampMax = GetData<UPropertyBagHierarchyNumericProperty>()->bUseClampMax, ClampMax = GetData<UPropertyBagHierarchyNumericProperty>()->ClampMax](FPropertyBagPropertyDesc& Desc)
+					{
+						if (bUseClampMax)
+						{
+							Desc.SetMetaData(UE::StructUtils::Metadata::Specifiers::ClampMaxName, FString::SanitizeFloat(ClampMax));
+						}
+						else
+						{
+							Desc.RemoveMetadata(UE::StructUtils::Metadata::Specifiers::ClampMaxName);
+						}
+					});
+		}
+		
+		if (PropertyChangedEvent.Property->GetName() == GET_MEMBER_NAME_CHECKED(UPropertyBagHierarchyNumericProperty, ClampMax))
+		{
+			UE::StructUtils::ApplyChangesToSinglePropertyDesc(
+				LOCTEXT("ChangedClampMax", "Changed Clamp Max"),
+				*GetPropertyDesc(),
+				Cast<UPropertyBagHierarchyViewModel>(HierarchyViewModel)->GetPropertyBagHandle(),
+				[ClampMax = GetData<UPropertyBagHierarchyNumericProperty>()->ClampMax](FPropertyBagPropertyDesc& Desc)
+					{
+						Desc.SetMetaData(UE::StructUtils::Metadata::Specifiers::ClampMaxName, FString::SanitizeFloat(ClampMax));
+					});
+		}
+		
+		if (PropertyChangedEvent.Property->GetName() == GET_MEMBER_NAME_CHECKED(UPropertyBagHierarchyNumericProperty, bUseUIMin))
+		{
+			UE::StructUtils::ApplyChangesToSinglePropertyDesc(
+				LOCTEXT("ChangedUseUIMin", "Changed Use UI Min"),
+				*GetPropertyDesc(),
+				Cast<UPropertyBagHierarchyViewModel>(HierarchyViewModel)->GetPropertyBagHandle(),
+				[bUseUIMin = GetData<UPropertyBagHierarchyNumericProperty>()->bUseUIMin, UIMin = GetData<UPropertyBagHierarchyNumericProperty>()->UIMin](FPropertyBagPropertyDesc& Desc)
+					{
+						if (bUseUIMin)
+						{
+							Desc.SetMetaData(UE::StructUtils::Metadata::Specifiers::UIMinName, FString::SanitizeFloat(UIMin));
+						}
+						else
+						{
+							Desc.RemoveMetadata(UE::StructUtils::Metadata::Specifiers::UIMinName);
+						}
+					});
+		}
+		
+		if (PropertyChangedEvent.Property->GetName() == GET_MEMBER_NAME_CHECKED(UPropertyBagHierarchyNumericProperty, UIMin))
+		{
+			UE::StructUtils::ApplyChangesToSinglePropertyDesc(
+				LOCTEXT("ChangedUIMin", "Changed UI Min"),
+				*GetPropertyDesc(),
+				Cast<UPropertyBagHierarchyViewModel>(HierarchyViewModel)->GetPropertyBagHandle(),
+				[UIMin = GetData<UPropertyBagHierarchyNumericProperty>()->UIMin](FPropertyBagPropertyDesc& Desc)
+					{
+						Desc.SetMetaData(UE::StructUtils::Metadata::Specifiers::UIMinName, FString::SanitizeFloat(UIMin));
+					});
+		}
+		
+		if (PropertyChangedEvent.Property->GetName() == GET_MEMBER_NAME_CHECKED(UPropertyBagHierarchyNumericProperty, bUseUIMax))
+		{
+			UE::StructUtils::ApplyChangesToSinglePropertyDesc(
+				LOCTEXT("ChangedUseUIMax", "Changed Use UI Max"),
+				*GetPropertyDesc(),
+				Cast<UPropertyBagHierarchyViewModel>(HierarchyViewModel)->GetPropertyBagHandle(),
+				[bUseUIMax = GetData<UPropertyBagHierarchyNumericProperty>()->bUseUIMax, UIMax = GetData<UPropertyBagHierarchyNumericProperty>()->UIMax](FPropertyBagPropertyDesc& Desc)
+					{
+						if (bUseUIMax)
+						{
+							Desc.SetMetaData(UE::StructUtils::Metadata::Specifiers::UIMaxName, FString::SanitizeFloat(UIMax));
+						}
+						else
+						{
+							Desc.RemoveMetadata(UE::StructUtils::Metadata::Specifiers::UIMaxName);
+						}
+					});
+		}
+		
+		if (PropertyChangedEvent.Property->GetName() == GET_MEMBER_NAME_CHECKED(UPropertyBagHierarchyNumericProperty, UIMax))
+		{
+			UE::StructUtils::ApplyChangesToSinglePropertyDesc(
+				LOCTEXT("ChangedUIMax", "Changed UI Max"),
+				*GetPropertyDesc(),
+				Cast<UPropertyBagHierarchyViewModel>(HierarchyViewModel)->GetPropertyBagHandle(),
+				[UIMax = GetData<UPropertyBagHierarchyNumericProperty>()->UIMax](FPropertyBagPropertyDesc& Desc)
+					{
+						Desc.SetMetaData(UE::StructUtils::Metadata::Specifiers::UIMaxName, FString::SanitizeFloat(UIMax));
+					});
+		}
+		
+		if (PropertyChangedEvent.Property->GetName() == GET_MEMBER_NAME_CHECKED(UPropertyBagHierarchyNumericProperty, bUseUnit))
+		{
+			UE::StructUtils::ApplyChangesToSinglePropertyDesc(
+				LOCTEXT("ChangedUseUnit", "Changed Use Unit"),
+				*GetPropertyDesc(),
+				Cast<UPropertyBagHierarchyViewModel>(HierarchyViewModel)->GetPropertyBagHandle(),
+				[bUseUnits = GetData<UPropertyBagHierarchyNumericProperty>()->bUseUnit, Unit = GetData<UPropertyBagHierarchyNumericProperty>()->Unit](FPropertyBagPropertyDesc& Desc)
+					{
+						if (bUseUnits)
+						{
+							Desc.SetMetaData(UE::StructUtils::Metadata::Specifiers::UnitsName, FUnitConversion::GetUnitDisplayString(Unit));
+						}
+						else
+						{
+							Desc.RemoveMetadata(UE::StructUtils::Metadata::Specifiers::UnitsName);
+						}
+					});
+		}
+		
+		if (PropertyChangedEvent.Property->GetName() == GET_MEMBER_NAME_CHECKED(UPropertyBagHierarchyNumericProperty, Unit))
+		{
+			UE::StructUtils::ApplyChangesToSinglePropertyDesc(
+				LOCTEXT("ChangedUnit", "Changed Unit"),
+				*GetPropertyDesc(),
+				Cast<UPropertyBagHierarchyViewModel>(HierarchyViewModel)->GetPropertyBagHandle(),
+				[Unit = GetData<UPropertyBagHierarchyNumericProperty>()->Unit](FPropertyBagPropertyDesc& Desc)
+					{
+						Desc.SetMetaData(UE::StructUtils::Metadata::Specifiers::UnitsName, FUnitConversion::GetUnitDisplayString(Unit));
+					});
+		}
+	}
+	
+	// This will refresh the details panel
+	FPropertyBagHierarchyPropertyViewModel::PostEditChange(PropertyChangedEvent, PropertyThatChanged);
+}
+
+void UPropertyBagHierarchyViewModel::Initialize(TSharedRef<IPropertyHandle> InOriginalPropertyBagHandle, TSharedRef<IPropertyUtilities> InPropertyUtilities)
+{
+	FPropertyRowGeneratorArgs Args;
+	PropertyRowGenerator = FModuleManager::Get().LoadModuleChecked<FPropertyEditorModule>("PropertyEditor").CreatePropertyRowGenerator(Args);
+	
+	PropertyUtilities = InPropertyUtilities;
+	PropertyBagSchema = UE::StructUtils::ExtractPropertyBagSchemaCDO(InOriginalPropertyBagHandle);
+	HierarchyRoot = UE::StructUtils::ExtractHierarchyRoot(InOriginalPropertyBagHandle);
+	
+	TArray<UObject*> OuterObjects;
+	InOriginalPropertyBagHandle->GetOuterObjects(OuterObjects);
+	
+	if (OuterObjects.Num() != 1)
+	{
+		return;
+	}
+	
+	OwningObject = OuterObjects[0];
+	
+	// The original handle will become invalidated if the details panel does a full refresh. We need to ensure it stays valid.
+	// We cache the property path which allows us to find our generated handle again
+	PropertyPathToPropertyBag = InOriginalPropertyBagHandle->CreateFPropertyPath();
+	
+	// We generate our own property handle now. Even this one can become invalid if changes to the property bag are made, but we can regenerate it now
+	PropertyRowGenerator->SetObjects(OuterObjects);
+	RegeneratePropertyBagHandle();	
+	
+	// We add the properties that don't currently exist in the hierarchy to the root.
+	AddMissingPropertiesToRoot(false);
+	
+	Super::Initialize();
+}
+
+void UPropertyBagHierarchyViewModel::RegeneratePropertyBagHandle()
+{
+	if (!PropertyPathToPropertyBag.IsValid() || !PropertyRowGenerator.IsValid())
+	{
+		return;
+	}
+	
+	if (TSharedPtr<IDetailTreeNode> TreeNode = FindNode(PropertyRowGenerator->GetRootTreeNodes(), PropertyPathToPropertyBag.ToSharedRef()))
+	{
+		GeneratedPropertyBagHandle = TreeNode->CreatePropertyHandle();
+	}
+}
+
+const FPropertyBagPropertyDesc* UPropertyBagHierarchyViewModel::GetPropertyDescForProperty(FGuid PropertyId) const
+{
+	if (const FInstancedPropertyBag* PropertyBag = GetPropertyBag())
+	{
+		if (const UPropertyBag* PropertyBagStruct = PropertyBag->GetPropertyBagStruct())
+		{
+			return PropertyBagStruct->FindPropertyDescByID(PropertyId);
+		}
+	}
+	
+	return nullptr;
+}
+
+TConstArrayView<FPropertyBagPropertyDesc> UPropertyBagHierarchyViewModel::GetPropertyDescs() const
+{
+	if (const FInstancedPropertyBag* PropertyBag = GetPropertyBag())
+	{
+		if (const UPropertyBag* PropertyBagStruct = PropertyBag->GetPropertyBagStruct())
+		{
+			return PropertyBagStruct->GetPropertyDescs();
+		}
+	}
+	
+	return {};
+}
+
+void UPropertyBagHierarchyViewModel::RebuildDetails() const
+{
+	// We can't request a rebuild from the handle itself since it's not the same handle as in the details panel.
+	// However, the property utilities remain stable, so we can force a refresh. This would have also invalidated the original handle.
+	if (PropertyUtilities.IsValid())
+	{
+		PropertyUtilities.Pin()->RequestForceRefresh();
+	}
+}
+
+void UPropertyBagHierarchyViewModel::Tick(float DeltaTime)
+{
+	// If our handle becomes invalid, we regenerate it. This can happen when adding new properties for example
+	if (!GeneratedPropertyBagHandle.IsValid() || !GeneratedPropertyBagHandle->IsValidHandle())
+	{
+		RegeneratePropertyBagHandle();
+	}
+}
+
+bool UPropertyBagHierarchyViewModel::IsTickable() const
+{
+	return !HasAnyFlags(RF_ClassDefaultObject);
+}
+
+TStatId UPropertyBagHierarchyViewModel::GetStatId() const
+{
+	RETURN_QUICK_DECLARE_CYCLE_STAT(PropertyBagHierarchyViewModel, STATGROUP_Tickables);
+}
+
+UHierarchyRoot* UPropertyBagHierarchyViewModel::GetHierarchyRoot() const
+{
+	return HierarchyRoot.Get();
+}
+
+void UPropertyBagHierarchyViewModel::OnHierarchyChanged(const TInstancedStruct<FHierarchyChangedPayload>& Payload)
+{
+	Super::OnHierarchyChanged(Payload);
+	
+	if(Payload.IsValid())
+	{
+		if(Payload.GetScriptStruct() == FHierarchyChangedPayload_ElementsDeleted::StaticStruct())
+		{
+			TArray<TSharedPtr<FHierarchyElementViewModel>> DeletedElementViewModel = Payload.GetPtr<FHierarchyChangedPayload_ElementsDeleted>()->DeletedElementViewModels;
+			
+			bool bContainsHierarchyElement = DeletedElementViewModel.ContainsByPredicate([](TSharedPtr<FHierarchyElementViewModel> Candidate)
+			{
+				return Candidate->GetData()->IsA<UHierarchyCategory>() || Candidate->GetData()->IsA<UHierarchyItem>();
+			});
+			
+			// By deleting a category we might have deleted contained properties. Readd them to the hierarchy.
+			// Deleting properties isn't possible to do directly, but it might happen (e.g.: changing property type that warrants a new hierarchy type)
+			if (bContainsHierarchyElement)
+			{
+				AddMissingPropertiesToRoot(true);
+			}
+		}
+	}
+	
+	RebuildDetails();
+}
+
+TSubclassOf<UHierarchyCategory> UPropertyBagHierarchyViewModel::GetCategoryDataClass() const
+{
+	if (PropertyBagSchema.IsValid())
+	{
+		return PropertyBagSchema->GetHierarchyCategoryType();
+	}
+	
+	return nullptr;
+}
+
+TSubclassOf<UHierarchySection> UPropertyBagHierarchyViewModel::GetSectionDataClass() const
+{
+	if (PropertyBagSchema.IsValid())
+	{
+		return PropertyBagSchema->GetHierarchySectionType();
+	}
+	
+	return nullptr;
+}
+
+void UPropertyBagHierarchyViewModel::PrepareSourceItems(UHierarchyRoot* SourceRoot, TSharedPtr<FHierarchyRootViewModel> SourceRootViewModel)
+{
+	if (const FInstancedPropertyBag* PropertyBag = GetPropertyBag())
+	{
+		if (PropertyBagSchema.IsValid())
+		{
+			for (const FPropertyBagPropertyDesc& Desc : PropertyBag->GetPropertyBagStruct()->GetPropertyDescs())
+			{
+				TSubclassOf<UPropertyBagHierarchyProperty> PropertyClass = PropertyBagSchema->GetHierarchyPropertyType(Desc);
+				if (ensure(PropertyClass))
+				{
+					UPropertyBagHierarchyProperty* HierarchyProperty = NewObject<UPropertyBagHierarchyProperty>(SourceRoot, PropertyClass);
+					HierarchyProperty->Initialize(Desc);
+					
+					SourceRoot->GetChildrenMutable().Add(HierarchyProperty);
+				}
+			}
+			
+			SourceRootViewModel->SyncViewModelsToData();
+		}
+	}
+}
+
+TSharedPtr<FHierarchyElementViewModel> UPropertyBagHierarchyViewModel::CreateCustomViewModelForElement(UHierarchyElement* Element, TSharedPtr<FHierarchyElementViewModel> Parent)
+{
+	if (UPropertyBagHierarchyNumericProperty* NumericProperty = Cast<UPropertyBagHierarchyNumericProperty>(Element))
+	{
+		return MakeShared<FPropertyBagHierarchyNumericPropertyViewModel>(NumericProperty, Parent.ToSharedRef(), this);
+	}
+	
+	if (UPropertyBagHierarchyProperty* Property = Cast<UPropertyBagHierarchyProperty>(Element))
+	{
+		return MakeShared<FPropertyBagHierarchyPropertyViewModel>(Property, Parent.ToSharedRef(), this);
+	}
+	
+	return nullptr;
+}
+
+void UPropertyBagHierarchyViewModel::PostUndo(bool bSuccess)
+{
+	RebuildDetails();
+	
+	Super::PostUndo(bSuccess);
+}
+
+void UPropertyBagHierarchyViewModel::PostRedo(bool bSuccess)
+{
+	RebuildDetails();
+	
+	Super::PostRedo(bSuccess);
+}
+
+const FInstancedPropertyBag* UPropertyBagHierarchyViewModel::GetPropertyBag() const
+{
+	if (PropertyPathToPropertyBag.IsValid() && OwningObject.IsValid())
+	{
+		return PropertyPathToPropertyBag->GetLeafMostProperty().Property->ContainerPtrToValuePtr<FInstancedPropertyBag>(OwningObject.Get());
+	}
+	
+	if (GeneratedPropertyBagHandle.IsValid() && GeneratedPropertyBagHandle->IsValidHandle())
+	{
+		void* DataPtr = nullptr;
+		if (GeneratedPropertyBagHandle->GetValueData(DataPtr) == FPropertyAccess::Result::Success)
+		{
+			return (const FInstancedPropertyBag*) DataPtr;
+		}
+	}
+	
+	return nullptr; 
+}
+
+void UPropertyBagHierarchyViewModel::AddMissingPropertiesToRoot(bool bSendNotification)
+{
+	if (!HierarchyRoot.IsValid())
+	{
+		return;
+	}
+	
+	TArray<UPropertyBagHierarchyProperty*> NewProperties;
+	
+	if (const FInstancedPropertyBag* PropertyBag = GetPropertyBag())
+	{
+		TArray<const UPropertyBagHierarchyProperty*> AllContainedProperties;
+		HierarchyRoot->GetChildrenOfType(AllContainedProperties, true);
+		
+		if (PropertyBagSchema.IsValid())
+		{
+			for (const FPropertyBagPropertyDesc& Desc : PropertyBag->GetPropertyBagStruct()->GetPropertyDescs())
+			{
+				if (!AllContainedProperties.ContainsByPredicate([Desc](const UPropertyBagHierarchyProperty* Candidate)
+				{
+					return Candidate->GetPropertyId() == Desc.ID;
+				}))
+				{
+					TSubclassOf<UPropertyBagHierarchyProperty> NewPropertyType = PropertyBagSchema->GetHierarchyPropertyType(Desc);
+					
+					if (NewPropertyType)
+					{
+						UPropertyBagHierarchyProperty* NewHierarchyProperty = NewObject<UPropertyBagHierarchyProperty>(HierarchyRoot.Get(), NewPropertyType);
+						NewHierarchyProperty->Initialize(Desc);
+						HierarchyRoot->GetChildrenMutable().Add(NewHierarchyProperty);
+						NewProperties.Add(NewHierarchyProperty);
+					}
+				}
+			}
+		}		
+	}
+	
+	if (NewProperties.Num() > 0)
+	{
+		if (HierarchyRootViewModel.IsValid())
+		{
+			HierarchyRootViewModel->SyncViewModelsToData();
+		}
+		
+		if (bSendNotification)
+		{
+			FNotificationInfo Info(FText::FormatOrdered(LOCTEXT("AddedMissingPropertiesToRoot", "Added missing {0}|plural(one=property, other=properties) to the hierarchy root."), NewProperties.Num()));
+			FSlateNotificationManager::Get().AddNotification(Info);
+		}
+	}
+}
+
+bool UPropertyBagHierarchyViewModel::FindPropertyHandleRecursive(const TSharedPtr<IPropertyHandle>& PropertyHandle, TSharedRef<FPropertyPath> PropertyPath)
+{
+	if (PropertyHandle && PropertyHandle->IsValidHandle())
+	{
+		uint32 ChildrenCount = 0;
+		PropertyHandle->GetNumChildren(ChildrenCount);
+		for (uint32 Index = 0; Index < ChildrenCount; ++Index)
+		{
+			TSharedPtr<IPropertyHandle> ChildHandle = PropertyHandle->GetChildHandle(Index);
+			if (FindPropertyHandleRecursive(ChildHandle, PropertyPath))
+			{
+				return true;
+			}
+		}
+
+		if (PropertyHandle->GetProperty())
+		{
+			if (FPropertyPath::AreEqual(PropertyHandle->CreateFPropertyPath().ToSharedRef(), PropertyPath))
+			{
+				return true;
+			}
+		}
+	}
+
+	return false;
+}
+
+TSharedPtr<IDetailTreeNode> UPropertyBagHierarchyViewModel::FindTreeNodeRecursive(const TSharedRef<IDetailTreeNode>& RootNode, TSharedRef<FPropertyPath> PropertyPath)
+{
+	TArray<TSharedRef<IDetailTreeNode>> Children;
+	RootNode->GetChildren(Children);
+	for (TSharedRef<IDetailTreeNode>& Child : Children)
+	{
+		TSharedPtr<IDetailTreeNode> FoundNode = FindTreeNodeRecursive(Child, PropertyPath);
+		if (FoundNode.IsValid())
+		{
+			return FoundNode;
+		}
+	}
+
+	TSharedPtr<IPropertyHandle> Handle = RootNode->CreatePropertyHandle();
+	if (FindPropertyHandleRecursive(Handle, PropertyPath))
+	{
+		return RootNode;
+	}
+
+	return nullptr;
+}
+
+TSharedPtr<IDetailTreeNode> UPropertyBagHierarchyViewModel::FindNode(const TArray<TSharedRef<IDetailTreeNode>>& RootNodes, TSharedRef<FPropertyPath> PropertyPath)
+{
+	for (const TSharedRef<IDetailTreeNode>& CategoryNode : RootNodes)
+	{
+		TSharedPtr<IDetailTreeNode> FoundNode = FindTreeNodeRecursive(CategoryNode, PropertyPath);
+		if (FoundNode.IsValid())
+		{
+			return FoundNode;
+		}
+	}
+
+	return nullptr;
+}
+
+#undef LOCTEXT_NAMESPACE
diff --git a/Engine/Source/Editor/StructUtilsEditor/Private/HierarchyEditor/SPropertyBagHierarchyEditor.cpp b/Engine/Source/Editor/StructUtilsEditor/Private/HierarchyEditor/SPropertyBagHierarchyEditor.cpp
new file mode 100644
--- /dev/null
+++ b/Engine/Source/Editor/StructUtilsEditor/Private/HierarchyEditor/SPropertyBagHierarchyEditor.cpp
@@ -0,0 +1,189 @@
+﻿// Copyright Epic Games, Inc. All Rights Reserved.
+
+#include "SPropertyBagHierarchyEditor.h"
+#include "PropertyBagDetails.h"
+#include "HierarchyEditor/PropertyBagHierarchyViewModel.h"
+#include "PropertyPath.h"
+#include "SlateOptMacros.h"
+#include "StructUtilsEditorUtilsPrivate.h"
+#include "StructUtilsMetadata.h"
+#include "Framework/Notifications/NotificationManager.h"
+#include "Kismet2/BlueprintEditorUtils.h"
+#include "Subsystems/AssetEditorSubsystem.h"
+#include "Widgets/SDataHierarchyEditor.h"
+#include "Widgets/Notifications/SNotificationList.h"
+#include "Editor.h"
+
+#define LOCTEXT_NAMESPACE "PropertyBagHierarchyEditor"
+
+void SPropertyBagHierarchyEditor::Construct(const FArguments& InArgs, TSharedRef<IPropertyHandle> InPropertyBagHandle, TSharedRef<IPropertyUtilities> InPropertyUtilities)
+{
+	OnCloseRequested = InArgs._OnCloseRequested;
+	
+	PropertyBagSchemaCDO = UE::StructUtils::ExtractPropertyBagSchemaCDO(InPropertyBagHandle);
+	HierarchyRoot = UE::StructUtils::ExtractHierarchyRoot(InPropertyBagHandle);
+	
+	if (!PropertyBagSchemaCDO.IsValid() || !HierarchyRoot.IsValid())
+	{
+		return;
+	}
+	
+	TArray<UObject*> OutOuterObjects;
+	InPropertyBagHandle->GetOuterObjects(OutOuterObjects);
+	Algo::Transform(OutOuterObjects, OuterObjects, [](UObject* Object)
+	{
+		return TWeakObjectPtr<UObject>(Object);
+	});
+	
+	HierarchyViewModel.Reset(NewObject<UPropertyBagHierarchyViewModel>(GetTransientPackage()));
+	HierarchyViewModel->Initialize(InPropertyBagHandle, InPropertyUtilities);
+	
+	if (GEditor)
+	{
+		if (UAssetEditorSubsystem* AssetEditorSubsystem = GEditor->GetEditorSubsystem<UAssetEditorSubsystem>())
+		{
+			AssetEditorSubsystem->OnAssetClosedInEditor().AddSP(this, &SPropertyBagHierarchyEditor::OnAssetEditorClosed);
+		}
+	}
+	
+	ChildSlot
+	[
+		SNew(SDataHierarchyEditor, HierarchyViewModel.Get())
+		.OnGenerateRowContentWidget(this, &SPropertyBagHierarchyEditor::OnGenerateRowContentWidget)
+	]; 
+}
+
+void SPropertyBagHierarchyEditor::Tick(const FGeometry& AllottedGeometry, const double InCurrentTime, const float InDeltaTime)
+{
+	if (!HierarchyViewModel.IsValid())
+	{
+		return;
+	}
+	
+	uint64 CurrentHash = UE::StructUtils::CalcPropertyDescArrayHash(HierarchyViewModel->GetPropertyDescs());
+	if (LastPropertyBagHash != CurrentHash)
+	{
+		HierarchyViewModel->AddMissingPropertiesToRoot(false);
+		HierarchyViewModel->RequestFullRefreshNextFrame();
+		LastPropertyBagHash = CurrentHash;
+	}
+}
+
+TSharedPtr<IPropertyHandle> SPropertyBagHierarchyEditor::GetPropertyBagHandle() const
+{ 
+	if (HierarchyViewModel.IsValid())
+	{
+		return HierarchyViewModel->GetPropertyBagHandle(); 
+	}
+	
+	return nullptr;
+}
+
+SPropertyBagHierarchyEditor::~SPropertyBagHierarchyEditor()
+{
+	if (GEditor)
+	{
+		if (UAssetEditorSubsystem* AssetEditorSubsystem = GEditor->GetEditorSubsystem<UAssetEditorSubsystem>())
+		{
+			AssetEditorSubsystem->OnAssetClosedInEditor().RemoveAll(this);
+		}
+	}
+	
+	if (HierarchyViewModel.IsValid())
+	{
+		HierarchyViewModel->Finalize();
+		HierarchyViewModel->ConditionalBeginDestroy();
+		HierarchyViewModel.Reset();
+	}
+}
+
+void SPropertyBagHierarchyEditor::NavigateToProperty(const FPropertyBagPropertyDesc& PropertyDesc, bool bAddIfNotFound)
+{
+	if (!HierarchyViewModel.IsValid())
+	{
+		return;
+	}
+	
+	FHierarchyElementIdentity Identity = UPropertyBagHierarchyProperty::ConstructIdentity(PropertyDesc);
+	TSharedPtr<FHierarchyElementViewModel> FoundViewModel = HierarchyViewModel->GetHierarchyRootViewModel()->FindViewModelForChild(Identity, true);
+	
+	if (!FoundViewModel.IsValid())
+	{
+		if (bAddIfNotFound)
+		{			
+			UHierarchyElement* Root = HierarchyViewModel->GetHierarchyRootViewModel()->GetDataMutable();
+			TSubclassOf<UPropertyBagHierarchyProperty> HierarchyPropertyType = PropertyBagSchemaCDO->GetHierarchyPropertyType(PropertyDesc);
+			
+			if (HierarchyPropertyType)
+			{
+				UPropertyBagHierarchyProperty* NewHierarchyProperty = NewObject<UPropertyBagHierarchyProperty>(Root, HierarchyPropertyType);
+				NewHierarchyProperty->Initialize(PropertyDesc);
+				Root->GetChildrenMutable().Add(NewHierarchyProperty);
+				HierarchyViewModel->GetHierarchyRootViewModel()->SyncViewModelsToData();
+				FoundViewModel = HierarchyViewModel->GetHierarchyRootViewModel()->FindViewModelForChild(Identity, false);
+				
+				HierarchyViewModel->SetActiveHierarchySection(nullptr);
+				
+				HierarchyViewModel->RebuildDetails();
+				
+				FNotificationInfo Info(FText::FormatOrdered(LOCTEXT("MissingHierarchyPropertyAdded", "Property {0} was added to the hierarchy so it can be edited."), 
+					FText::FromName(PropertyDesc.Name)));
+				
+				FSlateNotificationManager::Get().AddNotification(Info);
+			}
+		}
+	}
+	
+	// We delay navigation by a short duration; the window containing the hierarchy editor might have been added this same frame
+	// in which case the navigation would fail. 0.25 seems like a good magic number
+	RegisterActiveTimer(0.25f, FWidgetActiveTimerDelegate::CreateLambda([this, Identity](double InCurrentTime, float InDeltaTime)
+	{
+		HierarchyViewModel->NavigateToElementInHierarchy(HierarchyViewModel->GetHierarchyRootViewModel()->FindViewModelForChild(Identity, true).ToSharedRef());
+			
+		return EActiveTimerReturnType::Stop;
+	}));	
+}
+
+TSharedRef<SWidget> SPropertyBagHierarchyEditor::OnGenerateRowContentWidget(TSharedRef<FHierarchyElementViewModel> ElementViewModel) const
+{
+	if (ElementViewModel->GetDataMutable()->IsA<UPropertyBagHierarchyProperty>())
+	{
+		TSharedRef<FPropertyBagHierarchyPropertyViewModel> HierarchyPropertyViewModel = StaticCastSharedRef<FPropertyBagHierarchyPropertyViewModel>(ElementViewModel);
+				
+		const FPropertyBagPropertyDesc* Desc = HierarchyPropertyViewModel->GetPropertyDesc();
+		FEdGraphPinType PinType = Desc ? UE::StructUtils::GetPropertyDescAsPin(*Desc) : FEdGraphPinType();
+		
+		return SNew(SHorizontalBox)
+			+ SHorizontalBox::Slot()
+			.AutoWidth()
+			.Padding(2.f)
+			[
+				SNew(SImage)
+				.Image(FBlueprintEditorUtils::GetIconFromPin(PinType))
+				.ColorAndOpacity(PropertyBagSchemaCDO->GetPinTypeColor(PinType))
+			]
+			+ SHorizontalBox::Slot()
+			[
+				SNew(SHierarchyElement, HierarchyPropertyViewModel)
+			];
+	}
+	
+	if(ElementViewModel->GetDataMutable()->IsA<UPropertyBagHierarchyCategory>())
+	{
+		TSharedRef<FHierarchyCategoryViewModel> HierarchyCategoryViewModel = StaticCastSharedRef<FHierarchyCategoryViewModel>(ElementViewModel);
+		return SNew(SHierarchyElement, HierarchyCategoryViewModel)
+			.Style(&FDataHierarchyEditorStyle::Get().GetWidgetStyle<FInlineEditableTextBlockStyle>("HierarchyEditor.Category"));
+	}
+			
+	return SNew(STextBlock).Text(ElementViewModel->ToStringAsText());
+}
+
+void SPropertyBagHierarchyEditor::OnAssetEditorClosed(UObject* Object, IAssetEditorInstance* AssetEditorInstance) const
+{
+	if (OuterObjects.Contains(Object))
+	{
+		OnCloseRequested.ExecuteIfBound();
+	}
+}
+
+#undef LOCTEXT_NAMESPACE
diff --git a/Engine/Source/Editor/StructUtilsEditor/Private/HierarchyEditor/SPropertyBagHierarchyEditor.h b/Engine/Source/Editor/StructUtilsEditor/Private/HierarchyEditor/SPropertyBagHierarchyEditor.h
new file mode 100644
--- /dev/null
+++ b/Engine/Source/Editor/StructUtilsEditor/Private/HierarchyEditor/SPropertyBagHierarchyEditor.h
@@ -0,0 +1,55 @@
+﻿// Copyright Epic Games, Inc. All Rights Reserved.
+
+#pragma once
+
+#include "CoreMinimal.h"
+#include "DataHierarchyViewModelBase.h"
+#include "PropertyBagDetails.h"
+#include "EdGraph/EdGraphSchema.h"
+#include "HierarchyEditor/PropertyBagHierarchyViewModel.h"
+#include "Subsystems/AssetEditorSubsystem.h"
+#include "Templates/SubclassOf.h"
+#include "Widgets/SCompoundWidget.h"
+
+class UPropertyBagSchema;
+class UPropertyBagHierarchyViewModel;
+
+/** A wrapper around the actual hierarchy editor. Owns the view model. */
+class STRUCTUTILSEDITOR_API SPropertyBagHierarchyEditor : public SCompoundWidget
+{
+public:
+	SLATE_BEGIN_ARGS(SPropertyBagHierarchyEditor)
+		{}
+		SLATE_EVENT(FSimpleDelegate, OnCloseRequested)
+	SLATE_END_ARGS()
+
+	/**
+	 * @param InArgs 
+	 * @param InPropertyBagHandle For initialization purposes only as it is expected to get invalidated. Internally we store another handle that we generate. 
+	 * @param InPropertyUtilities Required to refresh the details panel that originally spawned this hierarchy editor
+	 */
+	void Construct(const FArguments& InArgs, TSharedRef<IPropertyHandle> InPropertyBagHandle, TSharedRef<IPropertyUtilities> InPropertyUtilities);
+	virtual ~SPropertyBagHierarchyEditor() override;
+	virtual void Tick(const FGeometry& AllottedGeometry, const double InCurrentTime, const float InDeltaTime) override;
+	
+	/** Returns the property handle that we generated ourselves via property row generator. */
+	TSharedPtr<IPropertyHandle> GetPropertyBagHandle() const;
+	void NavigateToProperty(const FPropertyBagPropertyDesc& PropertyDesc, bool bAddIfNotFound = false);
+
+private:
+	TSharedRef<SWidget> OnGenerateRowContentWidget(TSharedRef<FHierarchyElementViewModel> ElementViewModel) const;
+	void OnAssetEditorClosed(UObject* Object, IAssetEditorInstance* AssetEditorInstance) const;
+	
+private:
+	TWeakObjectPtr<const UPropertyBagSchema> PropertyBagSchemaCDO;
+	TStrongObjectPtr<UPropertyBagHierarchyViewModel> HierarchyViewModel;
+	
+	TWeakObjectPtr<UHierarchyRoot> HierarchyRoot;
+	
+	TArray<TWeakObjectPtr<UObject>> OuterObjects;
+	
+	/** We recalculate this on tick. When differences are detected, we refresh the hierarchy editor. */
+	uint64 LastPropertyBagHash = 0;
+	
+	FSimpleDelegate OnCloseRequested;
+};
diff --git a/Engine/Source/Editor/StructUtilsEditor/Public/HierarchyEditor/PropertyBagHierarchyViewModel.h b/Engine/Source/Editor/StructUtilsEditor/Public/HierarchyEditor/PropertyBagHierarchyViewModel.h
new file mode 100644
--- /dev/null
+++ b/Engine/Source/Editor/StructUtilsEditor/Public/HierarchyEditor/PropertyBagHierarchyViewModel.h
@@ -0,0 +1,188 @@
+﻿// Copyright Epic Games, Inc. All Rights Reserved.
+
+#pragma once
+
+#include "CoreMinimal.h"
+#include "DataHierarchyViewModelBase.h"
+#include "Math/UnitConversion.h"
+#include "StructUtils/PropertyBag.h"
+#include "PropertyBagHierarchyViewModel.generated.h"
+
+class UPropertyBagSchema;
+
+/** This file defines hierarchy elements and the view model that controls the hierarchy.
+ *  If you are looking to make use of them, simply subclass UPropertyBagSchema and override relevant functions.
+ *  If you wish to add properties to your usecase for editor purposes, consider subclassing either property, section or category.
+ *  If your property is generic and should be functioning for every property bag, consider adding them in the base classes here.
+ */
+UCLASS(MinimalAPI)
+class UPropertyBagHierarchyRoot : public UHierarchyRoot
+{
+	GENERATED_BODY()
+	
+public:
+	// Add generic future properties for the root here
+};
+
+UCLASS(MinimalAPI)
+class UPropertyBagHierarchyCategory : public UHierarchyCategory
+{
+	GENERATED_BODY()
+	
+public:
+	// Add generic future properties for the category here
+};
+
+UCLASS(MinimalAPI)
+class UPropertyBagHierarchySection : public UHierarchySection
+{
+	GENERATED_BODY()
+	
+public:
+	// Add generic future properties for the section here
+};
+
+UCLASS(MinimalAPI)
+class UPropertyBagHierarchyProperty : public UHierarchyItem
+{
+	GENERATED_BODY()
+	
+public:
+	virtual void Initialize(const FPropertyBagPropertyDesc& PropertyDesc);
+	
+	static FHierarchyElementIdentity ConstructIdentity(const FPropertyBagPropertyDesc& PropertyDesc);
+	
+	FGuid GetPropertyId() const;
+
+	/** The tooltip for this property. */
+	UPROPERTY(EditAnywhere, Category="Generic", meta=(MultiLine))
+	FText Tooltip;
+	
+	/** Whether this property is advanced or not. If it's advanced, it will only display if the "Show Advanced" button is clicked. */
+	UPROPERTY(EditAnywhere, Category="Generic")
+	bool bAdvanced;
+};
+
+UCLASS(MinimalAPI)
+class UPropertyBagHierarchyNumericProperty : public UPropertyBagHierarchyProperty
+{
+	GENERATED_BODY()
+	
+public:
+	UPROPERTY()
+	bool bUseClampMin = false;
+	
+	UPROPERTY(EditAnywhere, Category="Numeric", meta=(EditCondition="bUseClampMin"))
+	float ClampMin = 0.f;
+
+	UPROPERTY()
+	bool bUseClampMax = false;
+	
+	UPROPERTY(EditAnywhere, Category="Numeric", meta=(EditCondition="bUseClampMax"))
+	float ClampMax = 10.f;
+	
+	UPROPERTY()
+	bool bUseUIMin = false;
+	
+	UPROPERTY(EditAnywhere, Category="Numeric", meta=(EditCondition="bUseUIMin"))
+	float UIMin = 0.f;
+	
+	UPROPERTY()
+	bool bUseUIMax = false;
+	
+	UPROPERTY(EditAnywhere, Category="Numeric", meta=(EditCondition="bUseUIMax"))
+	float UIMax = 10.f;
+	
+	UPROPERTY()
+	bool bUseUnit = false;
+	
+	UPROPERTY(EditAnywhere, Category="Numeric", meta=(EditCondition="bUseUnit"))
+	EUnit Unit = EUnit::Unspecified;
+};
+
+struct FPropertyBagHierarchyPropertyViewModel : public FHierarchyItemViewModel
+{
+	FPropertyBagHierarchyPropertyViewModel(UPropertyBagHierarchyProperty* InProperty, TSharedRef<FHierarchyElementViewModel> InParent, TWeakObjectPtr<UDataHierarchyViewModelBase> InHierarchyViewModel);
+	
+	const FPropertyBagPropertyDesc* GetPropertyDesc() const;
+	
+	virtual void PostEditChange(const FPropertyChangedEvent& PropertyChangedEvent, FProperty* PropertyThatChanged) override;
+	virtual bool RepresentsExternalData() const override { return true; }
+	virtual bool DoesExternalDataStillExist(const UHierarchyDataRefreshContext* Context) const override;
+	
+	/** Users can't manually delete the properties. They can still be deleted if their parent gets deleted however.
+	 *  We re-add the properties to the root in that case. */
+	virtual bool CanDeleteInternal() override { return false; }
+	
+protected:
+	virtual FString ToString() const override;
+};
+
+struct FPropertyBagHierarchyNumericPropertyViewModel : public FPropertyBagHierarchyPropertyViewModel
+{
+	FPropertyBagHierarchyNumericPropertyViewModel(UPropertyBagHierarchyNumericProperty* InProperty, TSharedRef<FHierarchyElementViewModel> InParent, TWeakObjectPtr<UDataHierarchyViewModelBase> InHierarchyViewModel);
+	
+	virtual void PostEditChange(const FPropertyChangedEvent& PropertyChangedEvent, FProperty* PropertyThatChanged) override;
+};
+
+/** The hierarchy view model for property bags. Defines core rules for how the hierarchy behaves. */
+UCLASS()
+class STRUCTUTILSEDITOR_API UPropertyBagHierarchyViewModel : public UDataHierarchyViewModelBase, public FTickableEditorObject
+{
+	GENERATED_BODY()
+	
+public:
+	/**
+	 * @param InOriginalPropertyBagHandle For initialization purposes only as it is expected to get invalidated. Internally we store another handle that we generate. 
+	 * @param InPropertyUtilities Required to refresh the details panel that originally spawned this hierarchy editor
+	 */
+	void Initialize(TSharedRef<IPropertyHandle> InOriginalPropertyBagHandle, TSharedRef<IPropertyUtilities> InPropertyUtilities);
+	
+	void RegeneratePropertyBagHandle();
+	
+	void AddMissingPropertiesToRoot(bool bSendNotification = false);
+	
+	const FPropertyBagPropertyDesc* GetPropertyDescForProperty(FGuid PropertyId) const;
+	TConstArrayView<FPropertyBagPropertyDesc> GetPropertyDescs() const;
+	
+	TSharedPtr<IPropertyHandle> GetPropertyBagHandle() const { return GeneratedPropertyBagHandle; }
+	TSharedPtr<IPropertyUtilities> GetPropertyUtilities() const { return PropertyUtilities.IsValid() ? PropertyUtilities.Pin() : nullptr; }
+	TSharedRef<FPropertyPath> GetCurrentPropertyPath() const { return PropertyPathToPropertyBag.ToSharedRef(); }
+	
+	void RebuildDetails() const;
+
+protected:
+	virtual void Tick(float DeltaTime) override;
+	virtual bool IsTickable() const override;
+	virtual TStatId GetStatId() const override;
+	
+	virtual UHierarchyRoot* GetHierarchyRoot() const override;
+	virtual void OnHierarchyChanged(const TInstancedStruct<FHierarchyChangedPayload>& Payload) override;
+	virtual TSubclassOf<UHierarchyCategory> GetCategoryDataClass() const override;
+	virtual TSubclassOf<UHierarchySection> GetSectionDataClass() const override;
+	virtual void PrepareSourceItems(UHierarchyRoot* SourceRoot, TSharedPtr<FHierarchyRootViewModel> SourceRootViewModel) override;
+	virtual TSharedPtr<FHierarchyElementViewModel> CreateCustomViewModelForElement(UHierarchyElement* Element, TSharedPtr<FHierarchyElementViewModel> Parent) override;
+	/** We auto-add elements to the hierarchy, so we don't need the source panel. */
+	virtual bool SupportsSourcePanel() const override { return false; }
+	virtual void PostUndo(bool bSuccess) override;
+	virtual void PostRedo(bool bSuccess) override;
+	
+	const FInstancedPropertyBag* GetPropertyBag() const;
+
+	static bool FindPropertyHandleRecursive(const TSharedPtr<IPropertyHandle>& PropertyHandle, TSharedRef<FPropertyPath> PropertyPath);
+	static TSharedPtr<IDetailTreeNode> FindTreeNodeRecursive(const TSharedRef<IDetailTreeNode>& RootNode, TSharedRef<FPropertyPath> PropertyPath);
+	static TSharedPtr<IDetailTreeNode> FindNode(const TArray<TSharedRef<IDetailTreeNode>>& RootNodes, TSharedRef<FPropertyPath> PropertyPath);
+	
+private:
+	TWeakPtr<IPropertyUtilities> PropertyUtilities;
+	TWeakObjectPtr<const UPropertyBagSchema> PropertyBagSchema;
+	
+	TSharedPtr<IPropertyRowGenerator> PropertyRowGenerator;
+	TWeakObjectPtr<UHierarchyRoot> HierarchyRoot;
+	TWeakObjectPtr<UObject> OwningObject;
+	TSharedPtr<FPropertyPath> PropertyPathToPropertyBag;
+	
+	uint64 LastHash = 0;
+	
+	TSharedPtr<IPropertyHandle> GeneratedPropertyBagHandle;
+};
diff --git a/Engine/Source/Editor/StructUtilsEditor/Public/PropertyBagDetails.h b/Engine/Source/Editor/StructUtilsEditor/Public/PropertyBagDetails.h
--- a/Engine/Source/Editor/StructUtilsEditor/Public/PropertyBagDetails.h
+++ b/Engine/Source/Editor/StructUtilsEditor/Public/PropertyBagDetails.h
@@ -13,6 +13,7 @@ class FReply;
 enum class ESelectorType : uint8;
 struct FEdGraphSchemaAction;
 struct FMenuEntryParams;
+struct EVisibility;
 
 class IPropertyHandle;
 class IPropertyUtilities;
@@ -20,6 +21,13 @@ class IDetailPropertyRow;
 class SInlineEditableTextBlock;
 class SWidget;
 
+class UPropertyBagSchema;
+class UHierarchySection;
+class UPropertyBagHierarchyRoot;
+class UPropertyBagHierarchyProperty;
+class UPropertyBagHierarchySection;
+class UPropertyBagHierarchyCategory;
+
 /**
  * The property bag details child rows can be completely customized by assigning a combination of these values
  * to their FPropertyBag 'ChildRowFeatures' metadata via the meta specifier. TODO: It isn't currently possible
@@ -77,16 +85,22 @@ namespace UE::StructUtils
 
 	/** @return property bag struct common to all edited properties. */
 	STRUCTUTILSEDITOR_API const UPropertyBag* GetCommonBagStruct(TSharedPtr<IPropertyHandle> StructProperty);
-
+	
+	STRUCTUTILSEDITOR_API TSubclassOf<UPropertyBagSchema> ExtractPropertyBagSchemaClass(TSharedRef<IPropertyHandle> InPropertyBagHandle);
+	STRUCTUTILSEDITOR_API const UPropertyBagSchema* ExtractPropertyBagSchemaCDO(TSharedRef<IPropertyHandle> InPropertyBagHandle);
+	STRUCTUTILSEDITOR_API UPropertyBagHierarchyRoot* ExtractHierarchyRoot(TSharedRef<IPropertyHandle> InPropertyBagHandle);
+	
 	/**
 	 * Allows the mutation of the property descriptors on a Property Bag property via the provided callback.
 	 * Handles the whole transaction internally, including building a new Property Bag Struct (if necessary).
 	 *
 	 * @param SessionName Name of the transaction
 	 * @param StructProperty The Property Bag handle
-	 * @param PropUtils Property Utilities, usually taken from the StructCustomizationUtils in CustomizeHeader
 	 * @param Function The callback that does the mutation
 	 */
+	STRUCTUTILSEDITOR_API void ApplyChangesToPropertyDescs(const FText& SessionName, const TSharedPtr<IPropertyHandle>& StructProperty, TFunctionRef<void(TArray<FPropertyBagPropertyDesc>&)> Function);
+
+	UE_DEPRECATED(5.8, "Please use the version without the property utilities parameter.")
 	STRUCTUTILSEDITOR_API void ApplyChangesToPropertyDescs(const FText& SessionName, const TSharedPtr<IPropertyHandle>& StructProperty, const TSharedPtr<IPropertyUtilities>& PropUtils, TFunctionRef<void(TArray<FPropertyBagPropertyDesc>&)> Function);
 
 	/**
@@ -96,10 +110,9 @@ namespace UE::StructUtils
 	 * @param SessionName Name of the transaction
 	 * @param PropertyHandle A child handle of the Property Bag to mutate
 	 * @param StructProperty The Property Bag handle
-	 * @param PropUtils Property Utilities, usually taken from the StructCustomizationUtils in CustomizeHeader
 	 * @param Function The callback that does the mutation
 	 */
-	STRUCTUTILSEDITOR_API void ApplyChangesToSinglePropertyDesc(const FText& SessionName, const TSharedPtr<IPropertyHandle> PropertyHandle, const TSharedPtr<IPropertyHandle>& StructProperty, const TSharedPtr<IPropertyUtilities>& PropUtils, TFunctionRef<void(FPropertyBagPropertyDesc&)> Function);
+	STRUCTUTILSEDITOR_API void ApplyChangesToSinglePropertyDesc(const FText& SessionName, const TSharedPtr<IPropertyHandle> PropertyHandle, const TSharedPtr<IPropertyHandle>& StructProperty, TFunctionRef<void(FPropertyBagPropertyDesc&)> Function);
 
 	/**
 	 * Allows the mutation of a single property descriptor on a Property Bag property via the provided callback.
@@ -108,9 +121,11 @@ namespace UE::StructUtils
 	 * @param SessionName Name of the transaction
 	 * @param PropertyDesc The child descriptor to edit
 	 * @param StructProperty The Property Bag handle
-	 * @param PropUtils Property Utilities, usually taken from the StructCustomizationUtils in CustomizeHeader
 	 * @param Function The callback that does the mutation
 	 */
+	STRUCTUTILSEDITOR_API void ApplyChangesToSinglePropertyDesc(const FText& SessionName, const FPropertyBagPropertyDesc& PropertyDesc, const TSharedPtr<IPropertyHandle>& StructProperty, TFunctionRef<void(FPropertyBagPropertyDesc&)> Function);
+	
+	UE_DEPRECATED(5.8, "Please use the version without the property utilities parameter.")
 	STRUCTUTILSEDITOR_API void ApplyChangesToSinglePropertyDesc(const FText& SessionName, const FPropertyBagPropertyDesc& PropertyDesc, const TSharedPtr<IPropertyHandle>& StructProperty, const TSharedPtr<IPropertyUtilities>& PropUtils, TFunctionRef<void(FPropertyBagPropertyDesc&)> Function);
 }
 
@@ -140,20 +155,31 @@ protected:
 	/** IPropertyTypeCustomization interface */
 	UE_API virtual void CustomizeHeader(TSharedRef<IPropertyHandle> StructPropertyHandle, FDetailWidgetRow& HeaderRow, IPropertyTypeCustomizationUtils& StructCustomizationUtils) override;
 	UE_API virtual void CustomizeChildren(TSharedRef<IPropertyHandle> StructPropertyHandle, IDetailChildrenBuilder& StructBuilder, IPropertyTypeCustomizationUtils& StructCustomizationUtils) override;
-
+	
+	static FReply SummonHierarchyEditor(TSharedPtr<IPropertyHandle> InPropertyBagHandle, TSharedPtr<IPropertyUtilities> InPropertyUtilities);
+	void NavigateToHierarchyProperty(const FPropertyBagPropertyDesc& Desc);
+	
+	static void ClosePropertyBagHierarchyEditorWindow();
 private:
-	/** Handle to the struct property being edited */
+	/** Handle to the property bag being edited */
 	TSharedPtr<IPropertyHandle> StructProperty;
 	TSharedPtr<IPropertyUtilities> PropUtils;
 	EPropertyBagPropertyType DefaultType = EPropertyBagPropertyType::Bool;
 	EPropertyBagChildRowFeatures ChildRowFeatures = EPropertyBagChildRowFeatures::Default;
+	
+	// We keep track of existing window & hierarchy editor loosely so we can recycle them where appropriate
+	// This also means we can only have 1 property bag hierarchy editor open at a time
+	// This is because an additional window can't rely on the existing details panel and property handles too much
+	static TWeakPtr<SWindow> PropertyBagHierarchyEditorWindow;
+	static TWeakPtr<class SPropertyBagHierarchyEditor> PropertyBagHierarchyEditor;
 
 	UE_DEPRECATED(5.6, "Use 'ChildRowFeatures' instead.")
 	bool bFixedLayout = false;
 	bool bAllowContainers = true;
 	
 	// Default to UPropertyBagSchema, can be set with metadata PropertyBagSchemaClass
-	TSubclassOf<UEdGraphSchema> PropertyBagSchemaClass;
+	TWeakObjectPtr<const UPropertyBagSchema> PropertyBagSchemaCDO;
+	TWeakObjectPtr<UPropertyBagHierarchyRoot> HierarchyRoot;
 };
 
 /** 
@@ -167,13 +193,22 @@ private:
 class FPropertyBagInstanceDataDetails : public FInstancedStructDataDetails
 {
 public:
+	/** Generic delegate to execute on a property description */
+	DECLARE_DELEGATE_OneParam(FPropertyBagPropertyDelegate, const FPropertyBagPropertyDesc& PropertyDesc)
+	
 	struct FConstructParams
 	{
+		// Hierarchy Params are only meaningful if the owning property bag has a valid HierarchyRoot provided by a schema
+		struct FHierarchyParams
+		{
+			FPropertyBagPropertyDelegate OnNavigateToHierarchyPropertyRequested;
+		};
+		
 		TSharedPtr<IPropertyHandle> BagStructProperty = nullptr;
 		TSharedPtr<IPropertyUtilities> PropUtils = nullptr;
 		bool bAllowContainers = true;
 		EPropertyBagChildRowFeatures ChildRowFeatures = EPropertyBagChildRowFeatures::Default;
-		TSubclassOf<UEdGraphSchema> PropertyBagSchemaClass;
+		FHierarchyParams HierarchyParams;
 	};
 
 	UE_API explicit FPropertyBagInstanceDataDetails(const FConstructParams& ConstructParams);
@@ -188,7 +223,8 @@ public:
 	FPropertyBagInstanceDataDetails& operator=(FPropertyBagInstanceDataDetails&& Other) noexcept = delete;
 	UE_API PRAGMA_ENABLE_DEPRECATION_WARNINGS
 
-	virtual void OnGroupRowAdded(IDetailGroup& GroupRow, int32 Level, const FString& Category) const override;
+	virtual void AddChildRows(IDetailChildrenBuilder& ChildBuilder, const TArray<TSharedPtr<IPropertyHandle>>& ChildProperties) override;
+	UE_API virtual void OnGroupRowAdded(IDetailGroup& GroupRow, int32 Level, const FString& Category) const override;
 	UE_API virtual void OnChildRowAdded(IDetailPropertyRow& ChildRow) override;
 	
 	/** Enum describing if a property is overridden, or undetermined (e.g. multiselection) */
@@ -262,7 +298,21 @@ public:
 
 	/** Called to reset the child property to default value. */
 	UE_API virtual void ResetToDefault(TSharedPtr<IPropertyHandle> ChildPropertyHandle);
-
+	
+	virtual void SetOnRebuildChildren(FSimpleDelegate InOnRegenerateChildren) override
+	{
+		OnRebuildChildren = InOnRegenerateChildren;
+	}
+	
+	/**~ Public Hierarchy Editor API */	
+	UE_API EVisibility ShouldShowHierarchyProperty(TSharedPtr<IPropertyHandle> ChildPropertyHandle) const;
+	
+	TWeakObjectPtr<const UHierarchySection> GetActiveHierarchySection() const { return HierarchyData.ActiveHierarchySection; }
+	
+	/** Retrieves the hierarchy property handle map. Only meaningful if used in conjunction with a hierarchy root. */
+	const TMap<const UPropertyBagHierarchyProperty*, TSharedPtr<IPropertyHandle>>& GetHierarchyPropertyHandleMap() const { return HierarchyData.HierarchyPropertyHandleMap; }
+	/**~ End of Hierarchy Editor API */
+	
 protected:
 	UE_DEPRECATED(5.6, "Replaced by 'CreatePropertyDetailsWidget' to allow for more customizable drop-down menu location and content.")
 	UE_API TSharedRef<SWidget> OnPropertyNameContent(TSharedPtr<IPropertyHandle> ChildPropertyHandle, TSharedPtr<SInlineEditableTextBlock> InlineWidget) const;
@@ -270,6 +320,8 @@ protected:
 	TSharedPtr<IPropertyHandle> BagStructProperty;
 	TSharedPtr<IPropertyUtilities> PropUtils;
 
+	FSimpleDelegate OnRebuildChildren;
+	
 	bool bAllowContainers = true;
 	EPropertyBagChildRowFeatures ChildRowFeatures = EPropertyBagChildRowFeatures::Default;
 
@@ -278,8 +330,35 @@ protected:
 	UE_DEPRECATED(5.6, "Use 'bAllowContainers' instead.")
 	bool bAllowArrays = true;
 
-	// Default to UPropertyBagSchema, can be set with metadata PropertyBagSchemaClass
+	UE_DEPRECATED(5.8, "Use PropertyBagSchema instead. Will be invalid. PropertyBagSchema is set automatically via property metadata instead of passed in via construction arguments now.")
 	TSubclassOf<UEdGraphSchema> PropertyBagSchemaClass;
+	
+	// Default to UPropertyBagSchema, can be set with metadata PropertyBagSchemaClass
+	TWeakObjectPtr<const UPropertyBagSchema> PropertyBagSchema;
+	
+	/**~ Hierarchy Editor API */
+	void GenerateHierarchySectionRow(IDetailChildrenBuilder& ChildrenBuilder);
+	void GenerateHierarchyRootRows(IDetailChildrenBuilder& ChildrenBuilder);
+	void GenerateHierarchyLeftoverPropertyRows(IDetailChildrenBuilder& ChildrenBuilder, const TArray<TSharedPtr<IPropertyHandle>>& AllChildPropertyHandles);
+	void GenerateHierarchyShowAdvancedRow(IDetailChildrenBuilder& ChildrenBuilder);
+	
+	const FSlateBrush* GetAdvancedImage() const;
+	EVisibility GetAdvancedButtonVisibility() const;
+	
+	TSharedPtr<SWidget> ShowAdvancedButton;
+	bool bShouldShowAdvanced = false;
+	
+	struct FHierarchyData
+	{
+		TWeakObjectPtr<UPropertyBagHierarchyRoot> HierarchyRoot;
+		TMap<const UPropertyBagHierarchyProperty*, TSharedPtr<IPropertyHandle>> HierarchyPropertyHandleMap;
+		TWeakObjectPtr<const UHierarchySection> ActiveHierarchySection;
+	
+		FPropertyBagPropertyDelegate OnNavigateToHierarchyPropertyRequestedDelegate;
+	};
+	
+	FHierarchyData HierarchyData;
+	/**~ End of Hierarchy Editor API */	
 };
 
 UCLASS(MinimalAPI)
@@ -303,6 +382,21 @@ class UPropertyBagSchema : public UEdGraphSchema_K2
 	GENERATED_BODY()
 public:
 	UE_API virtual bool SupportsPinTypeContainer(TWeakPtr<const FEdGraphSchemaAction> SchemaAction, const FEdGraphPinType& PinType, const EPinContainerType& ContainerType) const override;
+	
+	/** Retrieves the hierarchy root from the given objects that own the property bag. Optional. */
+	UE_API virtual UPropertyBagHierarchyRoot* GetHierarchyRoot(const TArray<UObject*>& ObjectsWithProperty) const;
+
+	/** Defines the property type within the hierarchy. Useful if you want to subclass and add your own properties. */
+	UE_API virtual TSubclassOf<UPropertyBagHierarchyProperty> GetHierarchyPropertyType(const FPropertyBagPropertyDesc& Desc) const;
+	
+	/** Defines the category type within the hierarchy. Useful if you want to subclass and add your own properties. */
+	UE_API virtual TSubclassOf<UPropertyBagHierarchyCategory> GetHierarchyCategoryType() const;
+	
+	/** Defines the section type within the hierarchy. Useful if you want to subclass and add your own properties. */
+	UE_API virtual TSubclassOf<UPropertyBagHierarchySection> GetHierarchySectionType() const;
+	
+	/** Post Load helper function. Override if you know what you are doing. */
+	UE_API virtual void TransferPropertyBagMetadataIntoHierarchy(const FInstancedPropertyBag& PropertyBag, UPropertyBagHierarchyRoot& HierarchyRoot) const; 
 };
 
 #undef UE_API
diff --git a/Engine/Source/Editor/StructUtilsEditor/StructUtilsEditor.Build.cs b/Engine/Source/Editor/StructUtilsEditor/StructUtilsEditor.Build.cs
--- a/Engine/Source/Editor/StructUtilsEditor/StructUtilsEditor.Build.cs
+++ b/Engine/Source/Editor/StructUtilsEditor/StructUtilsEditor.Build.cs
@@ -29,6 +29,7 @@ namespace UnrealBuildTool.Rules
 					"AIModule",
 					"DetailCustomizations",
 					"ComponentVisualizers",
+					"DataHierarchyEditor"
 				}
 			);
 
@@ -42,7 +43,7 @@ namespace UnrealBuildTool.Rules
 					"AIGraph",
 					"ToolMenus",
 					"BlueprintGraph",
-					"KismetCompiler",
+					"KismetCompiler"
 				}
 			);
 		}
diff --git a/Engine/Source/Runtime/CoreUObject/Private/StructUtils/PropertyBag.cpp b/Engine/Source/Runtime/CoreUObject/Private/StructUtils/PropertyBag.cpp
--- a/Engine/Source/Runtime/CoreUObject/Private/StructUtils/PropertyBag.cpp
+++ b/Engine/Source/Runtime/CoreUObject/Private/StructUtils/PropertyBag.cpp
@@ -66,80 +66,6 @@ namespace UE::StructUtils::Private
 		return From != nullptr && To != nullptr && From->IsChildOf(To);
 	}
 
-	void AppendPropertyDescs(FXxHash64Builder& HashBuilder, const FPropertyBagPropertyDesc& Desc);
-
-	void AppendPathName(FStringBuilderBase& StringBuilder, const UObject* Object)
-	{
-		if (Object)
-		{
-			Object->GetPathName(nullptr, StringBuilder);
-		}
-	}
-
-	void AppendValueObjectHash(FXxHash64Builder& HashBuilder, FStringBuilderBase& StringBuilder, const UObject* InValueObject)
-	{
-		AppendPathName(StringBuilder, InValueObject);
-#if WITH_EDITOR
-		// for user defined structures we need to hash each property. after changing a user defined structure
-		// we may find the same property bag again if the hash is not sufficient, even though the memory layout
-		// has changed. hashing the path name of the user defined structure is not enough.
-		if(const UUserDefinedStruct* UserDefinedStruct = Cast<UUserDefinedStruct>(InValueObject))
-		{
-			for(const FField* Property = UserDefinedStruct->ChildProperties; Property != nullptr; Property = Property->Next)
-			{
-				const uint32 FieldHash = GetTypeHash(Property);
-				HashBuilder.Update(&FieldHash, sizeof(FieldHash));
-				if(const FStructProperty* StructProperty = CastField<FStructProperty>(Property))
-				{
-					if(StructProperty->Struct)
-					{
-						AppendValueObjectHash(HashBuilder, StringBuilder, StructProperty->Struct);
-					}
-				}
-			}
-		}
-		// for property bags we recurse
-		else if(const UPropertyBag* PropertyBag = Cast<UPropertyBag>(InValueObject))
-		{
-			for (const FPropertyBagPropertyDesc& Desc : PropertyBag->GetPropertyDescs())
-			{
-				AppendPropertyDescs(HashBuilder, Desc);
-			}
-		}
-#endif
-	}
-
-	void AppendPropertyDescs(FXxHash64Builder& HashBuilder, const FPropertyBagPropertyDesc& Desc)
-	{
-		TStringBuilder<256> StringBuilder;
-		HashBuilder.Update(&Desc.ID, sizeof(Desc.ID));
-		StringBuilder << Desc.Name;
-		HashBuilder.Update(&Desc.ValueType, sizeof(Desc.ValueType));
-		HashBuilder.Update(&Desc.ContainerTypes, sizeof(Desc.ContainerTypes));
-		HashBuilder.Update(&Desc.PropertyFlags, sizeof(Desc.PropertyFlags));
-		HashBuilder.Update(&Desc.KeyType, sizeof(Desc.KeyType));
-#if WITH_EDITORONLY_DATA
-		for (const FPropertyBagPropertyDescMetaData& MetaData : Desc.MetaData)
-		{
-			StringBuilder << MetaData.Key;
-			HashBuilder.Update(*MetaData.Value, MetaData.Value.Len() * sizeof(FString::ElementType));
-		}
-#endif
-		AppendValueObjectHash(HashBuilder, StringBuilder, Desc.ValueTypeObject);
-		AppendValueObjectHash(HashBuilder, StringBuilder, Desc.KeyTypeObject);
-		HashBuilder.Update(StringBuilder.ToString(), StringBuilder.Len() * sizeof(TStringBuilder<256>::ElementType));
-	}
-
-	uint64 CalcPropertyDescArrayHash(const TConstArrayView<FPropertyBagPropertyDesc> Descs)
-	{
-		FXxHash64Builder HashBuilder;
-		for (const FPropertyBagPropertyDesc& Desc : Descs)
-		{
-			AppendPropertyDescs(HashBuilder, Desc);
-		}
-		return HashBuilder.Finalize().Hash;
-	}
-
 	bool IsTypeHashable(const EPropertyBagPropertyType InType, const UObject* InTypeObject)
 	{
 		switch (InType)
@@ -4430,7 +4356,7 @@ TArray<FPropertyBagPropertyDesc> UPropertyBag::SanitizePropertyDescs(const TArra
 const UPropertyBag* UPropertyBag::GetOrCreateFromDescs(const TConstArrayView<FPropertyBagPropertyDesc> InPropertyDescs, const TCHAR* PrefixName)
 {
 	TArray<FPropertyBagPropertyDesc> PropertyDescs = SanitizePropertyDescs(InPropertyDescs);
-	const uint64 BagHash = UE::StructUtils::Private::CalcPropertyDescArrayHash(PropertyDescs);
+	const uint64 BagHash = UE::StructUtils::CalcPropertyDescArrayHash(PropertyDescs);
 	const FString ScriptStructName = PrefixName == nullptr
 		? FString::Printf(TEXT("PropertyBag_%llx"), BagHash)
 		: FString::Printf(TEXT("%s_%llx"), PrefixName, BagHash);
diff --git a/Engine/Source/Runtime/CoreUObject/Private/StructUtils/StructUtils.cpp b/Engine/Source/Runtime/CoreUObject/Private/StructUtils/StructUtils.cpp
new file mode 100644
--- /dev/null
+++ b/Engine/Source/Runtime/CoreUObject/Private/StructUtils/StructUtils.cpp
@@ -0,0 +1,82 @@
+﻿// Copyright Epic Games, Inc. All Rights Reserved.
+
+#include "StructUtils/StructUtils.h"
+
+#include "Hash/xxhash.h"
+#include "StructUtils/PropertyBag.h"
+#include "StructUtils/UserDefinedStruct.h"
+
+namespace UE::StructUtils::Private
+{
+	void AppendPathName(FStringBuilderBase& StringBuilder, const UObject* Object)
+	{
+		if (Object)
+		{
+			Object->GetPathName(nullptr, StringBuilder);
+		}
+	}
+
+	void AppendValueObjectHash(FXxHash64Builder& HashBuilder, FStringBuilderBase& StringBuilder, const UObject* InValueObject)
+	{
+		AppendPathName(StringBuilder, InValueObject);
+	#if WITH_EDITOR
+		// for user defined structures we need to hash each property. after changing a user defined structure
+		// we may find the same property bag again if the hash is not sufficient, even though the memory layout
+		// has changed. hashing the path name of the user defined structure is not enough.
+		if(const UUserDefinedStruct* UserDefinedStruct = Cast<UUserDefinedStruct>(InValueObject))
+		{
+			for(const FField* Property = UserDefinedStruct->ChildProperties; Property != nullptr; Property = Property->Next)
+			{
+				const uint32 FieldHash = GetTypeHash(Property);
+				HashBuilder.Update(&FieldHash, sizeof(FieldHash));
+				if(const FStructProperty* StructProperty = CastField<FStructProperty>(Property))
+				{
+					if(StructProperty->Struct)
+					{
+						AppendValueObjectHash(HashBuilder, StringBuilder, StructProperty->Struct);
+					}
+				}
+			}
+		}
+		// for property bags we recurse
+		else if(const UPropertyBag* PropertyBag = Cast<UPropertyBag>(InValueObject))
+		{
+			for (const FPropertyBagPropertyDesc& Desc : PropertyBag->GetPropertyDescs())
+			{
+				UE::StructUtils::AppendPropertyDescs(HashBuilder, Desc);
+			}
+		}
+	#endif
+	}
+}
+
+void UE::StructUtils::AppendPropertyDescs(FXxHash64Builder& HashBuilder, const FPropertyBagPropertyDesc& Desc)
+{
+	TStringBuilder<256> StringBuilder;
+	HashBuilder.Update(&Desc.ID, sizeof(Desc.ID));
+	StringBuilder << Desc.Name;
+	HashBuilder.Update(&Desc.ValueType, sizeof(Desc.ValueType));
+	HashBuilder.Update(&Desc.ContainerTypes, sizeof(Desc.ContainerTypes));
+	HashBuilder.Update(&Desc.PropertyFlags, sizeof(Desc.PropertyFlags));
+	HashBuilder.Update(&Desc.KeyType, sizeof(Desc.KeyType));
+#if WITH_EDITORONLY_DATA
+	for (const FPropertyBagPropertyDescMetaData& MetaData : Desc.MetaData)
+	{
+		StringBuilder << MetaData.Key;
+		HashBuilder.Update(*MetaData.Value, MetaData.Value.Len() * sizeof(FString::ElementType));
+	}
+#endif
+	Private::AppendValueObjectHash(HashBuilder, StringBuilder, Desc.ValueTypeObject);
+	Private::AppendValueObjectHash(HashBuilder, StringBuilder, Desc.KeyTypeObject);
+	HashBuilder.Update(StringBuilder.ToString(), StringBuilder.Len() * sizeof(TStringBuilder<256>::ElementType));
+}
+
+uint64 UE::StructUtils::CalcPropertyDescArrayHash(const TConstArrayView<FPropertyBagPropertyDesc> Descs)
+{
+	FXxHash64Builder HashBuilder;
+	for (const FPropertyBagPropertyDesc& Desc : Descs)
+	{
+		AppendPropertyDescs(HashBuilder, Desc);
+	}
+	return HashBuilder.Finalize().Hash;
+}
diff --git a/Engine/Source/Runtime/CoreUObject/Public/StructUtils/StructUtils.h b/Engine/Source/Runtime/CoreUObject/Public/StructUtils/StructUtils.h
--- a/Engine/Source/Runtime/CoreUObject/Public/StructUtils/StructUtils.h
+++ b/Engine/Source/Runtime/CoreUObject/Public/StructUtils/StructUtils.h
@@ -5,6 +5,12 @@
 #include "Templates/UnrealTypeTraits.h"
 #include "UObject/Class.h"
 
+class UPropertyBagHierarchyRoot;
+struct FInstancedPropertyBag;
+struct FPropertyBagPropertyDesc;
+class FXxHash64Builder;
+struct FPropertyBagPropertyDescMetaData;
+
 namespace UE::StructUtils
 {
 	template <typename T>
@@ -119,4 +125,8 @@ namespace UE::StructUtils
 		// Clamp count between 0 and the distance to the end of the range
 		InOutCount = FMath::Clamp(InOutCount, 0, (ArrayNum - InOutIndex));
 	}
+	
+	COREUOBJECT_API void AppendPropertyDescs(FXxHash64Builder& HashBuilder, const FPropertyBagPropertyDesc& Desc);
+
+	COREUOBJECT_API uint64 CalcPropertyDescArrayHash(const TConstArrayView<FPropertyBagPropertyDesc> Descs);
 }
