[Backout] - CL50304124
#rnx
Original CL Desc
-----------------------------------------------------------------
Added Data Hierarchy Editor support to FInstancedPropertyBag
-- PropertyBagSchema class can now direct the property bag details to an optional HierarchyRoot object
--- PropertyBagSchema class also contains information about hierarchy types
--- This can be used to subclass and add properties
-- Moved some util functions around hash calculation into public
-- Replaced some UEdGraphSchema with UPropertyBagSchema
- Added more encapsulation so that FPropertyBagInstanceDataDetails is more standalone and can be subclassed
-- Removed some of the construct params in favor of FPropertyBagInstanceDataDetails looking relevant data up via property handle.
-- Added InstanceMetaData support to allow forwarding metadata to an inner FInstancedPropertyBag

#Synced-CL 50298389

#Synced-CL 50298389

diff --git a/Engine/Source/Editor/EditorStyle/Private/StarshipStyle.cpp b/Engine/Source/Editor/EditorStyle/Private/StarshipStyle.cpp
--- a/Engine/Source/Editor/EditorStyle/Private/StarshipStyle.cpp
+++ b/Engine/Source/Editor/EditorStyle/Private/StarshipStyle.cpp
@@ -473,7 +473,6 @@ void FStarshipEditorStyle::FStyle::SetupGeneralStyles()
 		Set("Icons.Heart", new IMAGE_BRUSH_SVG("Starship/Common/Heart", Icon16x16));
 		Set("Icons.RecentAssets", new IMAGE_BRUSH_SVG("Starship/Common/RecentAssets", Icon16x16));
 		Set("Icons.Scalability", new IMAGE_BRUSH_SVG("Starship/Common/Scalability_16", Icon20x20));
-		Set("Icons.EditHierarchy", new IMAGE_BRUSH_SVG("Starship/Common/EditHierarchy_16", Icon16x16));
 		Set("Icons.ViewportScalability", new IMAGE_BRUSH_SVG("Starship/Common/ViewportScalability_16", Icon16x16));
 		Set("Icons.ViewportScalabilityReset", new IMAGE_BRUSH_SVG("Starship/Common/ViewportScalabilityReset_16", Icon16x16));
 		Set("Icons.EllipsisVerticalNarrow",
diff --git a/Engine/Source/Editor/StructUtilsEditor/Private/Customizations/PropertyBagDetails.cpp b/Engine/Source/Editor/StructUtilsEditor/Private/Customizations/PropertyBagDetails.cpp
--- a/Engine/Source/Editor/StructUtilsEditor/Private/Customizations/PropertyBagDetails.cpp
+++ b/Engine/Source/Editor/StructUtilsEditor/Private/Customizations/PropertyBagDetails.cpp
@@ -2,37 +2,30 @@
 
 #include "PropertyBagDetails.h"
 
-#include "DataHierarchyViewModelBase.h"
 #include "StructUtilsEditorUtilsPrivate.h"
 
 #include "DetailLayoutBuilder.h"
 #include "DetailWidgetRow.h"
-#include "Editor.h"
 #include "IDetailChildrenBuilder.h"
 #include "IDetailGroup.h"
 #include "IPropertyUtilities.h"
-#include "PropertyBagDragDropHandler.h"
 #include "PropertyCustomizationHelpers.h"
+#include "PropertyBagDragDropHandler.h"
 #include "ScopedTransaction.h"
 #include "SDraggableBox.h" // Custom widget for drag and drop sections
 #include "StructUtilsMetadata.h"
 #include "STypeSelector.h" // Custom widget for pill type selector
-#include "Framework/Application/SlateApplication.h"
 #include "Framework/MultiBox/MultiBoxBuilder.h"
-#include "HierarchyEditor/PropertyBagHierarchyViewModel.h"
-#include "HierarchyEditor/SPropertyBagHierarchyEditor.h"
-#include "Kismet2/BlueprintEditorUtils.h"
 #include "StructUtils/UserDefinedStruct.h"
+#include "Templates/ValueOrError.h"
 #include "UObject/EnumProperty.h"
-#include "Layout/Visibility.h"
+#include "UObject/TextProperty.h"
 
-#include "Widgets/SDataHierarchyEditor.h"
-#include "Widgets/Docking/SDockTab.h"
 #include "Widgets/Images/SImage.h"
 #include "Widgets/Input/SButton.h"
 #include "Widgets/Input/SComboButton.h"
+#include "Widgets/Input/SEditableTextBox.h"
 #include "Widgets/Input/SMultiLineEditableTextBox.h"
-#include "Widgets/Layout/SWrapBox.h"
 #include "Widgets/Text/SInlineEditableTextBlock.h"
 
 #include UE_INLINE_GENERATED_CPP_BY_NAME(PropertyBagDetails)
@@ -41,9 +34,6 @@
 
 ////////////////////////////////////
 
-TWeakPtr<SPropertyBagHierarchyEditor> FPropertyBagDetails::PropertyBagHierarchyEditor;
-TWeakPtr<SWindow> FPropertyBagDetails::PropertyBagHierarchyEditorWindow;
-
 namespace UE::StructUtils
 {
 /** Sets property descriptor based on a Blueprint pin type. */
@@ -605,6 +595,7 @@ void DeleteProperty(TSharedPtr<IPropertyHandle> InStructProperty, TSharedPtr<IPr
 	ApplyChangesToPropertyDescs(
 		FText::Format(LOCTEXT("OnPropertyDeleted", "Deleted property: {0}"), ChildPropertyHandle->GetPropertyDisplayName()),
 		InStructProperty,
+		PropUtils,
 		[&ChildPropertyHandle](TArray<FPropertyBagPropertyDesc>& PropertyDescs)
 		{
 			const FProperty* Property = ChildPropertyHandle ? ChildPropertyHandle->GetProperty() : nullptr;
@@ -642,6 +633,7 @@ void PinInfoChanged(TSharedPtr<IPropertyHandle> ChildPropertyHandle, const TShar
 	ApplyChangesToPropertyDescs(
 		FText::Format(LOCTEXT("OnPropertyTypeChanged", "Changed property type: {0}"), ChildPropertyHandle->GetPropertyDisplayName()),
 		InBagStructProperty,
+		InPropUtils,
 		[&PinType, &ChildPropertyHandle](TArray<FPropertyBagPropertyDesc>& PropertyDescs)
 		{
 			// Find and change struct type
@@ -676,7 +668,7 @@ FOptionalSize GetDesiredInputWidgetSizeByType(const TSharedPtr<IPropertyHandle>&
 
 TSharedRef<SWidget> CreateTypeSelectionWidget(TSharedPtr<IPropertyHandle> ChildPropertyHandle, const TSharedPtr<IPropertyHandle>& InBagStructProperty, const TSharedPtr<IPropertyUtilities>& InPropUtils, SPinTypeSelector::ESelectorType SelectorType, const bool bAllowContainers, TSubclassOf<UEdGraphSchema> PropertyBagSchemaClass)
 {
-	const UPropertyBagSchema* Schema = PropertyBagSchemaClass ? GetDefault<UPropertyBagSchema>(PropertyBagSchemaClass) : GetDefault<UPropertyBagSchema>();
+	const UEdGraphSchema* Schema = PropertyBagSchemaClass ? GetDefault<UEdGraphSchema>(PropertyBagSchemaClass) : GetDefault<UPropertyBagSchema>();
 
 	return SNew(SBox)
 		.HAlign(HAlign_Right)
@@ -732,65 +724,9 @@ const UPropertyBag* GetCommonBagStruct(TSharedPtr<IPropertyHandle> StructPropert
 	return CommonBagStruct;
 }
 
-TSubclassOf<UPropertyBagSchema> ExtractPropertyBagSchemaClass(TSharedRef<IPropertyHandle> InPropertyBagHandle)
-{
-	FString PropertyBagSchemaClassName;
-	
-	/** We first try to extract the instance metadata, if specified. */
-	if (const FString* InstancePropertyBagSchemaClassName = InPropertyBagHandle->GetInstanceMetaData(UE::StructUtils::Metadata::PropertyBagSchemaClassName))
-	{
-		PropertyBagSchemaClassName = *InstancePropertyBagSchemaClassName;
-	}
-	
-	/** If not available, we attempt to use the class metadata. */
-	if (PropertyBagSchemaClassName.IsEmpty())
-	{
-		if (FProperty* MetaDataProperty = InPropertyBagHandle->GetMetaDataProperty())
-		{
-			PropertyBagSchemaClassName = MetaDataProperty->GetMetaData(UE::StructUtils::Metadata::PropertyBagSchemaClassName);
-		}
-	}
-	
-	if (!PropertyBagSchemaClassName.IsEmpty())
-	{
-		return FindObject<UClass>(nullptr, *PropertyBagSchemaClassName);
-	}
-	
-	return nullptr;
-}
-
-const UPropertyBagSchema* ExtractPropertyBagSchemaCDO(TSharedRef<IPropertyHandle> InPropertyBagHandle)
-{
-	TSubclassOf<UPropertyBagSchema> SchemaClass = ExtractPropertyBagSchemaClass(InPropertyBagHandle);
-	
-	if (const UPropertyBagSchema* Schema = SchemaClass ? GetDefault<UPropertyBagSchema>(SchemaClass) : GetDefault<UPropertyBagSchema>())
-	{
-		return Schema;
-	}
-	
-	return nullptr;
-}
-
-UPropertyBagHierarchyRoot* ExtractHierarchyRoot(TSharedRef<IPropertyHandle> InPropertyBagHandle)
-{
-	if (!InPropertyBagHandle->IsValidHandle())
-	{
-		return nullptr;
-	}
-	
-	if (const UPropertyBagSchema* Schema = ExtractPropertyBagSchemaCDO(InPropertyBagHandle))
-	{
-		TArray<UObject*> OuterObjects;
-		InPropertyBagHandle->GetOuterObjects(OuterObjects);
-		return Schema->GetHierarchyRoot(OuterObjects);
-	}
-	
-	return nullptr;
-}
-
-void ApplyChangesToPropertyDescs(const FText& SessionName, const TSharedPtr<IPropertyHandle>& StructProperty, TFunctionRef<void(TArray<FPropertyBagPropertyDesc>&)> Function)
+void ApplyChangesToPropertyDescs(const FText& SessionName, const TSharedPtr<IPropertyHandle>& StructProperty, const TSharedPtr<IPropertyUtilities>& PropUtils, TFunctionRef<void(TArray<FPropertyBagPropertyDesc>&)> Function)
 {
-	if (!StructProperty || !StructProperty.Get()->IsValidHandle())
+	if (!StructProperty || !StructProperty.Get()->IsValidHandle() || !PropUtils)
 	{
 		return;
 	}
@@ -798,52 +734,39 @@ void ApplyChangesToPropertyDescs(const FText& SessionName, const TSharedPtr<IPro
 	FScopedTransaction Transaction(SessionName);
 	TArray<FPropertyBagPropertyDesc> PropertyDescs = Private::GetCommonPropertyDescs(StructProperty);
 	StructProperty->NotifyPreChange();
-	
-	if (UPropertyBagHierarchyRoot* Root = ExtractHierarchyRoot(StructProperty.ToSharedRef()))
-	{
-		Root->Modify();
-	}
-	
+
 	Function(PropertyDescs);
 
 	Private::SetPropertyDescs(StructProperty, PropertyDescs);
 
 	StructProperty->NotifyPostChange(EPropertyChangeType::ValueSet);
 	StructProperty->NotifyFinishedChangingProperties();
-	
-	StructProperty->RequestRebuildChildren();
-}
-
-void ApplyChangesToPropertyDescs(const FText& SessionName, const TSharedPtr<IPropertyHandle>& StructProperty, const TSharedPtr<IPropertyUtilities>& PropUtils, TFunctionRef<void(TArray<FPropertyBagPropertyDesc>&)> Function)
-{
-	ApplyChangesToPropertyDescs(SessionName, StructProperty, Function);
-}
-
-void ApplyChangesToSinglePropertyDesc(const FText& SessionName, const TSharedPtr<IPropertyHandle> PropertyHandle, const TSharedPtr<IPropertyHandle>& StructProperty, TFunctionRef<void(FPropertyBagPropertyDesc&)> Function)
-{
-	ApplyChangesToPropertyDescs(SessionName, StructProperty, [Function = std::move(Function), Property = PropertyHandle->GetProperty()](TArray<FPropertyBagPropertyDesc>& PropertyDescs)
-		{
-			if (FPropertyBagPropertyDesc* Desc = PropertyDescs.FindByPredicate([Property](const FPropertyBagPropertyDesc& OutDesc) { return OutDesc.CachedProperty == Property; }))
-			{
-				Function(*Desc);
-			}
-		});
+	if (PropUtils)
+	{
+		PropUtils->ForceRefresh();
+	}
 }
 
-void ApplyChangesToSinglePropertyDesc(const FText& SessionName, const FPropertyBagPropertyDesc& PropertyDesc, const TSharedPtr<IPropertyHandle>& StructProperty, TFunctionRef<void(FPropertyBagPropertyDesc&)> Function)
+void ApplyChangesToSinglePropertyDesc(const FText& SessionName, const TSharedPtr<IPropertyHandle> PropertyHandle, const TSharedPtr<IPropertyHandle>& StructProperty, const TSharedPtr<IPropertyUtilities>& PropUtils, TFunctionRef<void(FPropertyBagPropertyDesc&)> Function)
 {
-	ApplyChangesToPropertyDescs(SessionName, StructProperty, [Function = std::move(Function), &PropertyDesc](TArray<FPropertyBagPropertyDesc>& PropertyDescs)
+	ApplyChangesToPropertyDescs(SessionName, StructProperty, PropUtils, [Function = std::move(Function), Property = PropertyHandle->GetProperty()](TArray<FPropertyBagPropertyDesc>& PropertyDescs)
+	{
+		if (FPropertyBagPropertyDesc* Desc = PropertyDescs.FindByPredicate([Property](const FPropertyBagPropertyDesc& OutDesc) { return OutDesc.CachedProperty == Property; }))
 		{
-			if (FPropertyBagPropertyDesc* Desc = PropertyDescs.FindByPredicate([&PropertyDesc](const FPropertyBagPropertyDesc& OutDesc) { return OutDesc == PropertyDesc; }))
-			{
-				Function(*Desc);
-			}
-		});
+			Function(*Desc);
+		}
+	});
 }
 
 void ApplyChangesToSinglePropertyDesc(const FText& SessionName, const FPropertyBagPropertyDesc& PropertyDesc, const TSharedPtr<IPropertyHandle>& StructProperty, const TSharedPtr<IPropertyUtilities>& PropUtils, TFunctionRef<void(FPropertyBagPropertyDesc&)> Function)
 {
-	ApplyChangesToSinglePropertyDesc(SessionName, PropertyDesc, StructProperty, Function);
+	ApplyChangesToPropertyDescs(SessionName, StructProperty, PropUtils, [Function = std::move(Function), &PropertyDesc](TArray<FPropertyBagPropertyDesc>& PropertyDescs)
+	{
+		if (FPropertyBagPropertyDesc* Desc = PropertyDescs.FindByPredicate([&PropertyDesc](const FPropertyBagPropertyDesc& OutDesc) { return OutDesc == PropertyDesc; }))
+		{
+			Function(*Desc);
+		}
+	});
 }
 
 namespace Constants
@@ -858,85 +781,6 @@ namespace Constants
 }
 } // UE::StructUtils
 
-
-/** The category builder will display all sub-categories & parameters contained within a given category in the details panel. */
-class FHierarchyPropertyCategoryBuilder : public IDetailCustomNodeBuilder, public TSharedFromThis<FHierarchyPropertyCategoryBuilder>
-{
-public:
-	FHierarchyPropertyCategoryBuilder(const UPropertyBagHierarchyCategory& InHierarchyCategory, TSharedRef<FPropertyBagInstanceDataDetails> InOwningDetails);
-	
-	virtual void GenerateHeaderRowContent(FDetailWidgetRow& NodeRow) override;
-	
-	virtual bool InitiallyCollapsed() const override;
-
-	virtual FName GetName() const  override;
-
-	virtual void GenerateChildContent(IDetailChildrenBuilder& ChildrenBuilder) override;
-
-private:
-	TWeakObjectPtr<const UPropertyBagHierarchyCategory> HierarchyCategory;
-	FName CustomBuilderRowName;
-	TWeakPtr<FPropertyBagInstanceDataDetails> OwningDetails;
-};
-
-FHierarchyPropertyCategoryBuilder::FHierarchyPropertyCategoryBuilder(const UPropertyBagHierarchyCategory& InHierarchyCategory, TSharedRef<FPropertyBagInstanceDataDetails> InOwningDetails)
-	: HierarchyCategory(&InHierarchyCategory)
-	, CustomBuilderRowName(HierarchyCategory->GetCategoryName())
-	, OwningDetails(InOwningDetails)
-{	
-}
-
-void FHierarchyPropertyCategoryBuilder::GenerateHeaderRowContent(FDetailWidgetRow& NodeRow)
-{
-	NodeRow.NameContent()
-	[
-		SNew(STextBlock)
-		.Text_UObject(HierarchyCategory.Get(), &UHierarchyCategory::GetCategoryAsText)
-		.ToolTipText_UObject(HierarchyCategory.Get(), &UHierarchyCategory::GetTooltip)
-	];
-
-	NodeRow.FilterString(FText::FromName(CustomBuilderRowName));
-}
-
-bool FHierarchyPropertyCategoryBuilder::InitiallyCollapsed() const
-{
-	return false;
-}
-
-FName FHierarchyPropertyCategoryBuilder::GetName() const
-{
-	return CustomBuilderRowName;
-}
-
-void FHierarchyPropertyCategoryBuilder::GenerateChildContent(IDetailChildrenBuilder& ChildrenBuilder)
-{
-	if (!OwningDetails.IsValid())
-	{
-		return;
-	}
-	
-	TArray<const UPropertyBagHierarchyCategory*> ContainedCategories;
-	HierarchyCategory->GetChildrenOfType(ContainedCategories, false);
-	
-	for (const UPropertyBagHierarchyCategory* ContainedHierarchyCategory : ContainedCategories)
-	{
-		ChildrenBuilder.AddCustomBuilder(MakeShared<FHierarchyPropertyCategoryBuilder>(*ContainedHierarchyCategory, OwningDetails.Pin().ToSharedRef()));
-	}
-		
-	TArray<const UPropertyBagHierarchyProperty*> ContainedProperties;
-	HierarchyCategory->GetChildrenOfType<>(ContainedProperties, false);
-	
-	for (const UPropertyBagHierarchyProperty* ContainedHierarchyProperty : ContainedProperties)
-	{
-		if (OwningDetails.Pin()->GetHierarchyPropertyHandleMap().Contains(ContainedHierarchyProperty))
-		{
-			IDetailPropertyRow& PropertyRow = ChildrenBuilder.AddProperty(OwningDetails.Pin()->GetHierarchyPropertyHandleMap()[ContainedHierarchyProperty].ToSharedRef());
-			PropertyRow.Visibility(MakeAttributeSP(OwningDetails.Pin().Get(), &FPropertyBagInstanceDataDetails::ShouldShowHierarchyProperty, PropertyRow.GetPropertyHandle()));
-			OwningDetails.Pin()->OnChildRowAdded(PropertyRow);
-		}
-	}
-}
-
 //----------------------------------------------------------------//
 //  FPropertyBagInstanceDataDetails
 //  - StructProperty is FInstancedPropertyBag
@@ -954,21 +798,10 @@ PRAGMA_DISABLE_DEPRECATION_WARNINGS
 	, bFixedLayout(false)
 	, bAllowArrays(true)
 PRAGMA_ENABLE_DEPRECATION_WARNINGS
+	, PropertyBagSchemaClass(ConstructParams.PropertyBagSchemaClass)
 {
 	ensure(UE::StructUtils::Private::IsScriptStruct<FInstancedPropertyBag>(BagStructProperty));
 	ensure(PropUtils != nullptr);
-	
-	if (BagStructProperty.IsValid() && BagStructProperty->IsValidHandle())
-	{
-		PropertyBagSchema = UE::StructUtils::ExtractPropertyBagSchemaCDO(BagStructProperty.ToSharedRef());
-		
-		HierarchyData.HierarchyRoot = UE::StructUtils::ExtractHierarchyRoot(BagStructProperty.ToSharedRef());
-		
-		if (HierarchyData.HierarchyRoot.IsValid())
-		{
-			HierarchyData.OnNavigateToHierarchyPropertyRequestedDelegate = ConstructParams.HierarchyParams.OnNavigateToHierarchyPropertyRequested;
-		}
-	}
 }
 
 /** For backwards compatibility. */
@@ -987,93 +820,12 @@ PRAGMA_ENABLE_DEPRECATION_WARNINGS
 	ensure(PropUtils != nullptr);
 }
 
-void FPropertyBagInstanceDataDetails::GenerateHierarchyShowAdvancedRow(IDetailChildrenBuilder& ChildrenBuilder)
-{
-	if (!HierarchyData.HierarchyRoot.IsValid())
-	{
-		return;
-	}
-	
-	// At last we add a row for the 'Show Advanced' button. Only displayed if there is at least one advanced property
-	FDetailWidgetRow& AdvancedWidgetRow = ChildrenBuilder.AddCustomRow(INVTEXT("Advanced"));
-	AdvancedWidgetRow.Visibility(TAttribute<EVisibility>::CreateSP(this, &FPropertyBagInstanceDataDetails::GetAdvancedButtonVisibility));
-		
-	AdvancedWidgetRow.WholeRowContent()
-	.HAlign(HAlign_Center)
-	[
-		SAssignNew(ShowAdvancedButton, SButton)
-		.ButtonStyle(FAppStyle::Get(), "HoverHintOnly")
-		.ToolTipText(LOCTEXT("DisplayedAdvancedProperties", "Display advanced properties"))
-		.OnClicked_Lambda([this]()
-		{
-			bShouldShowAdvanced = !bShouldShowAdvanced;
-			return FReply::Handled();
-		})
-		[
-			SNew(SImage)
-			.Image(this, &FPropertyBagInstanceDataDetails::GetAdvancedImage)
-		]
-	];
-}
-
-void FPropertyBagInstanceDataDetails::AddChildRows(IDetailChildrenBuilder& ChildBuilder, const TArray<TSharedPtr<IPropertyHandle>>& ChildProperties)
-{
-	if (HierarchyData.HierarchyRoot.IsValid() && BagStructProperty.IsValid() && BagStructProperty->IsValidHandle())
-	{		
-		void* BagStructPtr = nullptr;
-		FPropertyAccess::Result Result = BagStructProperty->GetValueData(BagStructPtr);
-		
-		if (Result == FPropertyAccess::Fail)
-		{
-			return;
-		}
-		
-		FInstancedPropertyBag* Bag = static_cast<FInstancedPropertyBag*>(BagStructPtr);
-		
-		// We build a map from hierarchy property to property handle. This is then used in hierarchy category builders to retrieve their handles to add to the details panel.
-		TArray<const UPropertyBagHierarchyProperty*> AllHierarchyProperties;
-		HierarchyData.HierarchyRoot->GetChildrenOfType(AllHierarchyProperties, true);
-		
-		HierarchyData.HierarchyPropertyHandleMap.Empty();
-		for(const UPropertyBagHierarchyProperty* HierarchyProperty : AllHierarchyProperties)
-		{
-			if (const FPropertyBagPropertyDesc* PropertyDesc = Bag->FindPropertyDescByID(HierarchyProperty->GetPropertyId()))
-			{
-				const TSharedPtr<IPropertyHandle>* MatchingPropertyHandle = ChildProperties.FindByPredicate([PropertyDesc](const TSharedPtr<IPropertyHandle>& Candidate)
-				{
-					return Candidate->GetProperty() == PropertyDesc->CachedProperty;
-				});
-				
-				if (MatchingPropertyHandle)
-				{
-					HierarchyData.HierarchyPropertyHandleMap.Add(HierarchyProperty, *MatchingPropertyHandle);			
-				}
-			}
-		}
-		
-		// First, we generate the section row
-		GenerateHierarchySectionRow(ChildBuilder);
-		// Then the root level categories & properties
-		GenerateHierarchyRootRows(ChildBuilder);
-		// Then the leftover properties that have not been embedded into the hierarchy
-		GenerateHierarchyLeftoverPropertyRows(ChildBuilder, ChildProperties);
-		// Then the 'Show Advanced' row, if there is any advanced property
-		GenerateHierarchyShowAdvancedRow(ChildBuilder);
-	}
-	else
-	{
-		FInstancedStructDataDetails::AddChildRows(ChildBuilder, ChildProperties);
-	}
-}
-
 void FPropertyBagInstanceDataDetails::OnGroupRowAdded(IDetailGroup& GroupRow, int32 Level, const FString& Category) const
 {
 	using namespace UE::StructUtils;
 
 	FDetailWidgetRow& FolderRow = GroupRow.HeaderRow();
 	TWeakPtr<const FPropertyBagInstanceDataDetails> WeakSelf = SharedThis<const FPropertyBagInstanceDataDetails>(this);
-	TWeakPtr<IPropertyHandle> WeakPropertyBagHandle = BagStructProperty;
-	
 	FString FullCategoryName = GroupRow.GetGroupName().ToString();
 
 	/*** DRAG AND DROP HANDLER ***/
@@ -1110,13 +862,8 @@ void FPropertyBagInstanceDataDetails::OnGroupRowAdded(IDetailGroup& GroupRow, in
 
 		DragDropHandler->BindOnHandleDragDrop(
 			FOnPropertyBagDetailsRowDropOp::CreateLambda(
-				[WeakSelf, FullCategoryName, BagStructProperty = WeakPropertyBagHandle, PropUtils = PropUtils](const FPropertyBagPropertyDesc& DroppedPropertyDesc, const EItemDropZone DropZone) -> FReply
+				[WeakSelf, FullCategoryName, BagStructProperty = BagStructProperty, PropUtils = PropUtils](const FPropertyBagPropertyDesc& DroppedPropertyDesc, const EItemDropZone DropZone) -> FReply
 				{
-					if (!BagStructProperty.IsValid())
-					{
-						return FReply::Unhandled();
-					}
-						
 					if (ensure(DroppedPropertyDesc.CachedProperty && DropZone == EItemDropZone::OntoItem))
 					{
 						const TSharedPtr<const FPropertyBagInstanceDataDetails> DetailsSP = WeakSelf.Pin();
@@ -1130,7 +877,8 @@ void FPropertyBagInstanceDataDetails::OnGroupRowAdded(IDetailGroup& GroupRow, in
 						ApplyChangesToSinglePropertyDesc(
 							LOCTEXT("DragToChangeCategory", "Change property category"),
 							DroppedPropertyDesc,
-							BagStructProperty.Pin(),
+							BagStructProperty,
+							PropUtils,
 							[&DroppedPropertyDesc, &FullCategoryName](FPropertyBagPropertyDesc& Desc)
 							{
 								Metadata::SetCategory(Desc, FullCategoryName);
@@ -1185,6 +933,7 @@ void FPropertyBagInstanceDataDetails::OnGroupRowAdded(IDetailGroup& GroupRow, in
 				ApplyChangesToPropertyDescs(
 					LOCTEXT("InlineRenameCategory", "Rename category"),
 					BagStructProperty,
+					PropUtils,
 					[&InNewText, OldCategory = FullCategoryName, &Category](TArray<FPropertyBagPropertyDesc>& PropertyDescs)
 					{
 						FString NewCategory = OldCategory;
@@ -1239,6 +988,7 @@ void FPropertyBagInstanceDataDetails::OnGroupRowAdded(IDetailGroup& GroupRow, in
 					ApplyChangesToPropertyDescs(
 						LOCTEXT("OnCategoryDeleted", "Delete category"),
 						BagStructProperty,
+						PropUtils,
 						[GroupName](TArray<FPropertyBagPropertyDesc>& PropertyDescs)
 						{
 							for (FPropertyBagPropertyDesc& Desc : PropertyDescs)
@@ -1438,7 +1188,7 @@ void FPropertyBagInstanceDataDetails::OnChildRowAdded(IDetailPropertyRow& ChildR
 				.Padding(1, 0)
 				.AutoWidth()
 				[
-					CreateTypeSelectionWidget(ChildPropertyHandle, BagStructProperty, PropUtils, SPinTypeSelector::ESelectorType::Compact, bAllowContainers, PropertyBagSchema->GetClass())
+					CreateTypeSelectionWidget(ChildPropertyHandle, BagStructProperty, PropUtils, SPinTypeSelector::ESelectorType::Compact, bAllowContainers, PropertyBagSchemaClass)
 				];
 		}
 
@@ -1522,6 +1272,7 @@ void FPropertyBagInstanceDataDetails::OnChildRowAdded(IDetailPropertyRow& ChildR
 				ApplyChangesToPropertyDescs(
 					FText::Format(LOCTEXT("OnPropertyNameChanged", "Change property name: {0} -> {1}"), FText::FromName(OldName), FText::FromName(NewName)),
 					BagStructProperty,
+					PropUtils,
 					[&NewName, &ChildPropertyHandle](TArray<FPropertyBagPropertyDesc>& PropertyDescs)
 					{
 						const FProperty* Property = ChildPropertyHandle->GetProperty();
@@ -1602,6 +1353,7 @@ void FPropertyBagInstanceDataDetails::OnChildRowAdded(IDetailPropertyRow& ChildR
 							LOCTEXT("OnPropertyAccessSpecifierChanged", "Set access specifier."),
 							ChildPropertyHandle,
 							BagStructProperty,
+							PropUtils,
 							[Property, BagStructProperty](FPropertyBagPropertyDesc& Desc)
 							{
 								const bool bIsPrivate = !Property->HasAnyPropertyFlags(CPF_NativeAccessSpecifierPublic);
@@ -1659,7 +1411,7 @@ void FPropertyBagInstanceDataDetails::OnChildRowAdded(IDetailPropertyRow& ChildR
 						PropUtils,
 						SPinTypeSelector::ESelectorType::Full,
 						bAllowContainers,
-						PropertyBagSchema->GetClass()),
+						PropertyBagSchemaClass),
 					/*InLabel=*/FText::GetEmpty());
 				MenuBuilder.EndSection();
 			}
@@ -1702,114 +1454,145 @@ void FPropertyBagInstanceDataDetails::OnChildRowAdded(IDetailPropertyRow& ChildR
 
 				MenuBuilder.EndSection();
 			}
-			
-			if (HierarchyData.HierarchyRoot.IsValid())
+
+			// BagStructProperty and PropUtils are forwarded straight to OnTextCommitted, so they save a ref count and be refs here.
+			auto BuildMetadataSpecifierSubMenu = [](
+				TAttribute<FText>&& TextAttribute,
+				FOnVerifyTextChanged&& OnVerifyTextDelegate,
+				FOnTextCommitted&& OnTextCommittedDelegate) -> FNewMenuDelegate
+			{
+				TSharedRef<SWidget> TextBox = SNew(SMultiLineEditableTextBox)
+					.WrapTextAt(400)
+					.Text(MoveTemp(TextAttribute))
+					.OnVerifyTextChanged(MoveTemp(OnVerifyTextDelegate))
+					.OnTextCommitted(OnTextCommittedDelegate);
+
+				return FNewMenuDelegate::CreateLambda(
+					[TextBox](FMenuBuilder& OutMenuBuilder)
+					{
+						OutMenuBuilder.AddWidget(
+							SNew(SHorizontalBox)
+							+ SHorizontalBox::Slot()
+							.HAlign(HAlign_Fill)
+							.MinWidth(Constants::MinSubMenuTextBoxWidth)
+							.MaxWidth(Constants::MaxSubMenuTextBoxWidth)
+							.Padding(Constants::SubMenuTextBoxPadding, 0)
+							.AutoWidth()
+							[
+								TextBox
+							],
+							FText::GetEmpty(),
+							/*bNoIndent=*/true,
+							/*bInSearchable=*/true);
+					});
+			};
+
+			// The property's category (grouping) can be edited here.
+			if (EnumHasAllFlags(ChildRowFeatures, EPropertyBagChildRowFeatures::Categories | EPropertyBagChildRowFeatures::Menu_Categories))
 			{
-				if (HierarchyData.OnNavigateToHierarchyPropertyRequestedDelegate.IsBound())
+				MenuBuilder.BeginSection(/*InExtensionHook=*/NAME_None, LOCTEXT("DropDownMenuSectionCategory", "Category"));
+
+				if (ChildPropertyHandle->HasMetaData(Metadata::CategoryName))
 				{
 					MenuBuilder.AddMenuEntry(
-						LOCTEXT("DropDownMenuEditPropertyInHierarchyEditor", "Edit in Hierarchy Editor"),
-						LOCTEXT("DropDownMenUEditPropertyInHierarchyEditorToolTip", "Edit this property in the Hierarchy Editor"),
-						FSlateIcon(),
-						FUIAction(FExecuteAction::CreateLambda([this, PropertyDesc]()
+						LOCTEXT("DropDownMenuClearCategory", "Clear category"),
+						LOCTEXT("DropDownMenuClearCategoryToolTip", "Remove the property from its current category."),
+						FSlateIcon(FAppStyle::GetAppStyleSetName(), "Icons.Delete"),
+						FUIAction(FExecuteAction::CreateLambda([ChildPropertyHandle, StructProperty = BagStructProperty, PropUtils = PropUtils]()
 						{
-							HierarchyData.OnNavigateToHierarchyPropertyRequestedDelegate.Execute(*PropertyDesc);
-						}))
-					);
+							ApplyChangesToSinglePropertyDesc(
+								LOCTEXT("DropDownMenuOnCategoryCleared", "Clear property category"),
+								ChildPropertyHandle,
+								StructProperty,
+								PropUtils,
+								[](FPropertyBagPropertyDesc& Desc)
+								{
+									Metadata::RemoveCategory(Desc);
+								});
+						})));
 				}
-			}
 
-			// We only add metadata menu entries if we don't have a hierarchy root
-			if (!HierarchyData.HierarchyRoot.IsValid())
-			{
-				// PropertyBagHandle and PropUtils are forwarded straight to OnTextCommitted, so they save a ref count and be refs here.
-				auto BuildMetadataSpecifierSubMenu = [](
-					TAttribute<FText>&& TextAttribute,
-					FOnVerifyTextChanged&& OnVerifyTextDelegate,
-					FOnTextCommitted&& OnTextCommittedDelegate) -> FNewMenuDelegate
+				FOnVerifyTextChanged ValidateCategoryNameDelegate = FOnVerifyTextChanged::CreateLambda([](const FText& InText, FText& OutMessage)
 				{
-					TSharedRef<SWidget> TextBox = SNew(SMultiLineEditableTextBox)
-						.WrapTextAt(400)
-						.Text(MoveTemp(TextAttribute))
-						.OnVerifyTextChanged(MoveTemp(OnVerifyTextDelegate))
-						.OnTextCommitted(OnTextCommittedDelegate);
-
-					return FNewMenuDelegate::CreateLambda(
-						[TextBox](FMenuBuilder& OutMenuBuilder)
+					if (InText.ToString().Len() > Constants::MaxCategoryLength)
+					{
+						OutMessage = LOCTEXT("DropDownMenuInvalidCategoryName", "Invalid category name");
+						return false;
+					}
+					else
+					{
+						return true;
+					}
+				});
+
+				MenuBuilder.AddSubMenu(
+					FText::FromName(Metadata::CategoryName),
+					LOCTEXT("DropDownMenuSubMenuCategoryTooltip", "Set the category of this property. Subcategories can be created with the '|' character."),
+					BuildMetadataSpecifierSubMenu(
+						TAttribute<FText>::CreateLambda([ChildPropertyHandle]()
 						{
-							OutMenuBuilder.AddWidget(
-								SNew(SHorizontalBox)
-								+ SHorizontalBox::Slot()
-								.HAlign(HAlign_Fill)
-								.MinWidth(Constants::MinSubMenuTextBoxWidth)
-								.MaxWidth(Constants::MaxSubMenuTextBoxWidth)
-								.Padding(Constants::SubMenuTextBoxPadding, 0)
-								.AutoWidth()
-								[
-									TextBox
-								],
-								FText::GetEmpty(),
-								/*bNoIndent=*/true,
-								/*bInSearchable=*/true);
-						});
-				};
+							FText GroupLabel = FText::FromString("");
+							if (ChildPropertyHandle && ChildPropertyHandle->HasMetaData(Metadata::CategoryName))
+							{
+								GroupLabel = FText::FromString(ChildPropertyHandle->GetMetaData(Metadata::CategoryName));
+							}
 
-				// The property's category (grouping) can be edited here.
-				if (EnumHasAllFlags(ChildRowFeatures, EPropertyBagChildRowFeatures::Categories | EPropertyBagChildRowFeatures::Menu_Categories))
-				{
-					MenuBuilder.BeginSection(/*InExtensionHook=*/NAME_None, LOCTEXT("DropDownMenuSectionCategory", "Category"));
+							return GroupLabel;
+						}),
+						FOnVerifyTextChanged::CreateLambda([](const FText& InText, FText& OutMessage)
+						{
+							if (InText.ToString().Len() > Constants::MaxCategoryLength)
+							{
+								OutMessage = LOCTEXT("DropDownMenuInvalidCategoryName", "Invalid category name");
+								return false;
+							}
+							else
+							{
+								return true;
+							}
+						}),
 
-					if (ChildPropertyHandle->HasMetaData(Metadata::CategoryName))
-					{
-						MenuBuilder.AddMenuEntry(
-							LOCTEXT("DropDownMenuClearCategory", "Clear category"),
-							LOCTEXT("DropDownMenuClearCategoryToolTip", "Remove the property from its current category."),
-							FSlateIcon(FAppStyle::GetAppStyleSetName(), "Icons.Delete"),
-							FUIAction(FExecuteAction::CreateLambda([ChildPropertyHandle, StructProperty = BagStructProperty, PropUtils = PropUtils]()
+						FOnTextCommitted::CreateLambda([StructProperty = BagStructProperty, PropUtils = PropUtils, ChildPropertyHandle](const FText& CommittedText, const ETextCommit::Type CommitType)
+						{
+							if (CommitType == ETextCommit::OnEnter || CommitType == ETextCommit::OnUserMovedFocus)
 							{
 								ApplyChangesToSinglePropertyDesc(
-									LOCTEXT("DropDownMenuOnCategoryCleared", "Clear property category"),
+									LOCTEXT("DropDownMenuOnCategoryEdited", "Edit property category"),
 									ChildPropertyHandle,
 									StructProperty,
-									[](FPropertyBagPropertyDesc& Desc)
-										{
-											Metadata::RemoveCategory(Desc);
-										});
-							})));
-					}
+									PropUtils,
+									[&CommittedText](FPropertyBagPropertyDesc& Desc)
+									{
+										Metadata::SetCategory(Desc, CommittedText.ToString());
+									});
+							}
+						})),
+					/*bInOpenSubMenuOnClick=*/true);
 
-					FOnVerifyTextChanged ValidateCategoryNameDelegate = FOnVerifyTextChanged::CreateLambda([](const FText& InText, FText& OutMessage)
-					{
-						if (InText.ToString().Len() > Constants::MaxCategoryLength)
-						{
-							OutMessage = LOCTEXT("DropDownMenuInvalidCategoryName", "Invalid category name");
-							return false;
-						}
-						else
-						{
-							return true;
-						}
-					});
+				MenuBuilder.EndSection();
+			}
 
-					MenuBuilder.AddSubMenu(
-						FText::FromName(Metadata::CategoryName),
-						LOCTEXT("DropDownMenuSubMenuCategoryTooltip", "Set the category of this property. Subcategories can be created with the '|' character."),
-						BuildMetadataSpecifierSubMenu(
-							TAttribute<FText>::CreateLambda([ChildPropertyHandle]()
-							{
-								FText GroupLabel = FText::FromString("");
-								if (ChildPropertyHandle && ChildPropertyHandle->HasMetaData(Metadata::CategoryName))
-								{
-									GroupLabel = FText::FromString(ChildPropertyHandle->GetMetaData(Metadata::CategoryName));
-								}
+			// The property's optional metadata specifiers can be edited here.
+			if (EnumHasAllFlags(ChildRowFeatures, EPropertyBagChildRowFeatures::Menu_MetadataSpecifiers))
+			{
+				MenuBuilder.BeginSection(/*InExtensionHook=*/NAME_None, LOCTEXT("DropDownMenuSectionMetadataSpecifiers", "Specifiers"));
 
-								return GroupLabel;
-							}),
-							FOnVerifyTextChanged::CreateLambda([](const FText& InText, FText& OutMessage)
+				auto BuildGenericSpecifierSubMenu = [this, ChildPropertyHandle, BuildMetadataSpecifierSubMenu](const FName& SpecifierName, const bool bIsNumeric = false)
+				{
+					return BuildMetadataSpecifierSubMenu(
+						TAttribute<FText>::CreateLambda([SpecifierName, ChildPropertyHandle]()
+						{
+							return ChildPropertyHandle.IsValid()
+									   ? FText::FromString(ChildPropertyHandle->GetMetaData(SpecifierName))
+									   : FText::GetEmpty();
+						}),
+						!bIsNumeric
+							? FOnVerifyTextChanged()
+							: FOnVerifyTextChanged::CreateLambda([](const FText& InText, FText& OutMessage) -> bool
 							{
-								if (InText.ToString().Len() > Constants::MaxCategoryLength)
+								if (!InText.IsEmpty() && !InText.IsNumeric())
 								{
-									OutMessage = LOCTEXT("DropDownMenuInvalidCategoryName", "Invalid category name");
+									OutMessage = LOCTEXT("DropDownMenuSpecifierNumericValueExpected", "Value is not numeric");
 									return false;
 								}
 								else
@@ -1817,124 +1600,76 @@ void FPropertyBagInstanceDataDetails::OnChildRowAdded(IDetailPropertyRow& ChildR
 									return true;
 								}
 							}),
-
-							FOnTextCommitted::CreateLambda([StructProperty = BagStructProperty, PropUtils = PropUtils, ChildPropertyHandle](const FText& CommittedText, const ETextCommit::Type CommitType)
+						FOnTextCommitted::CreateLambda(
+							[SpecifierName, StructProperty = BagStructProperty, PropUtils = PropUtils, ChildPropertyHandle]
+							(const FText& CommittedText, const ETextCommit::Type CommitType)
 							{
-								if (CommitType == ETextCommit::OnEnter || CommitType == ETextCommit::OnUserMovedFocus)
+								if (CommitType != ETextCommit::OnCleared)
 								{
 									ApplyChangesToSinglePropertyDesc(
-										LOCTEXT("DropDownMenuOnCategoryEdited", "Edit property category"),
+										FText::Format(LOCTEXT("DropDownMenuOnSetSpecifier", "Set metadata specifier '{0}' value"), FText::FromName(SpecifierName)),
 										ChildPropertyHandle,
 										StructProperty,
-										[&CommittedText](FPropertyBagPropertyDesc& Desc)
+										PropUtils,
+										[SpecifierName, &CommittedText](FPropertyBagPropertyDesc& Desc)
+										{
+											// If empty, remove the specifier
+											if (CommittedText.IsEmpty())
+											{
+												Desc.RemoveMetadata(SpecifierName);
+											}
+											else
 											{
-												Metadata::SetCategory(Desc, CommittedText.ToString());
-											});
+												Desc.SetMetaData(SpecifierName, CommittedText.ToString());
+											}
+										});
 								}
-							})),
-						/*bInOpenSubMenuOnClick=*/true);
+							}));
+				};
 
-					MenuBuilder.EndSection();
-				}
+				MenuBuilder.AddSubMenu(
+					FText::FromName(Metadata::Specifiers::ToolTipName),
+					LOCTEXT("DropDownMenuSubMenuTooltipTooltip", "Set the tooltip for this property."),
+					BuildGenericSpecifierSubMenu(Metadata::Specifiers::ToolTipName),
+					/*bInOpenSubMenuOnClick=*/true);
 
-				// The property's optional metadata specifiers can be edited here.
-				if (EnumHasAllFlags(ChildRowFeatures, EPropertyBagChildRowFeatures::Menu_MetadataSpecifiers))
+				if (Metadata::Specifiers::SupportsClamping(*PropertyDesc))
 				{
-					MenuBuilder.BeginSection(/*InExtensionHook=*/NAME_None, LOCTEXT("DropDownMenuSectionMetadataSpecifiers", "Specifiers"));
-
-					auto BuildGenericSpecifierSubMenu = [this, ChildPropertyHandle, BuildMetadataSpecifierSubMenu](const FName& SpecifierName, const bool bIsNumeric = false)
-					{
-						return BuildMetadataSpecifierSubMenu(
-							TAttribute<FText>::CreateLambda([SpecifierName, ChildPropertyHandle]()
-							{
-								return ChildPropertyHandle.IsValid()
-										   ? FText::FromString(ChildPropertyHandle->GetMetaData(SpecifierName))
-										   : FText::GetEmpty();
-							}),
-							!bIsNumeric
-								? FOnVerifyTextChanged()
-								: FOnVerifyTextChanged::CreateLambda([](const FText& InText, FText& OutMessage) -> bool
-								{
-									if (!InText.IsEmpty() && !InText.IsNumeric())
-									{
-										OutMessage = LOCTEXT("DropDownMenuSpecifierNumericValueExpected", "Value is not numeric");
-										return false;
-									}
-									else
-									{
-										return true;
-									}
-								}),
-							FOnTextCommitted::CreateLambda(
-								[SpecifierName, StructProperty = BagStructProperty, PropUtils = PropUtils, ChildPropertyHandle]
-								(const FText& CommittedText, const ETextCommit::Type CommitType)
-								{
-									if (CommitType != ETextCommit::OnCleared)
-									{
-										ApplyChangesToSinglePropertyDesc(
-											FText::Format(LOCTEXT("DropDownMenuOnSetSpecifier", "Set metadata specifier '{0}' value"), FText::FromName(SpecifierName)),
-											ChildPropertyHandle,
-											StructProperty,
-											[SpecifierName, &CommittedText](FPropertyBagPropertyDesc& Desc)
-												{
-													// If empty, remove the specifier
-													if (CommittedText.IsEmpty())
-													{
-														Desc.RemoveMetadata(SpecifierName);
-													}
-													else
-													{
-														Desc.SetMetaData(SpecifierName, CommittedText.ToString());
-													}
-												});
-									}
-								}));
-					};
+					MenuBuilder.AddSubMenu(
+						FText::FromName(Metadata::Specifiers::ClampMinName),
+						LOCTEXT("DropDownMenuSubMenuClampMinTooltip", "Set the min value for this property."),
+						BuildGenericSpecifierSubMenu(Metadata::Specifiers::ClampMinName, /*bIsNumeric=*/true),
+						/*bInOpenSubMenuOnClick=*/true);
 
 					MenuBuilder.AddSubMenu(
-						FText::FromName(Metadata::Specifiers::ToolTipName),
-						LOCTEXT("DropDownMenuSubMenuTooltipTooltip", "Set the tooltip for this property."),
-						BuildGenericSpecifierSubMenu(Metadata::Specifiers::ToolTipName),
+						FText::FromName(Metadata::Specifiers::ClampMaxName),
+						LOCTEXT("DropDownMenuSubMenuClampMaxTooltip", "Set the max value for this property."),
+						BuildGenericSpecifierSubMenu(Metadata::Specifiers::ClampMaxName, /*bIsNumeric=*/true),
 						/*bInOpenSubMenuOnClick=*/true);
 
-					if (Metadata::Specifiers::SupportsClamping(*PropertyDesc))
-					{
-						MenuBuilder.AddSubMenu(
-							FText::FromName(Metadata::Specifiers::ClampMinName),
-							LOCTEXT("DropDownMenuSubMenuClampMinTooltip", "Set the min value for this property."),
-							BuildGenericSpecifierSubMenu(Metadata::Specifiers::ClampMinName, /*bIsNumeric=*/true),
-							/*bInOpenSubMenuOnClick=*/true);
-
-						MenuBuilder.AddSubMenu(
-							FText::FromName(Metadata::Specifiers::ClampMaxName),
-							LOCTEXT("DropDownMenuSubMenuClampMaxTooltip", "Set the max value for this property."),
-							BuildGenericSpecifierSubMenu(Metadata::Specifiers::ClampMaxName, /*bIsNumeric=*/true),
-							/*bInOpenSubMenuOnClick=*/true);
-
-						MenuBuilder.AddSubMenu(
-							FText::FromName(Metadata::Specifiers::UIMinName),
-							LOCTEXT("DropDownMenuSubMenuUIMinTooltip", "Set the min UI value for this property."),
-							BuildGenericSpecifierSubMenu(Metadata::Specifiers::UIMinName, /*bIsNumeric=*/true),
-							/*bInOpenSubMenuOnClick=*/true);
-
-						MenuBuilder.AddSubMenu(
-							FText::FromName(Metadata::Specifiers::UIMaxName),
-							LOCTEXT("DropDownMenuSubMenuUIMaxTooltip", "Set the max UI value for this property."),
-							BuildGenericSpecifierSubMenu(Metadata::Specifiers::UIMaxName, /*bIsNumeric=*/true),
-							/*bInOpenSubMenuOnClick=*/true);
-					}
+					MenuBuilder.AddSubMenu(
+						FText::FromName(Metadata::Specifiers::UIMinName),
+						LOCTEXT("DropDownMenuSubMenuUIMinTooltip", "Set the min UI value for this property."),
+						BuildGenericSpecifierSubMenu(Metadata::Specifiers::UIMinName, /*bIsNumeric=*/true),
+						/*bInOpenSubMenuOnClick=*/true);
 
-					if (Metadata::Specifiers::SupportsSettingUnits(*PropertyDesc))
-					{
-						MenuBuilder.AddSubMenu(
-							FText::FromName(Metadata::Specifiers::UnitsName),
-							LOCTEXT("DropDownMenuSubMenuUnitsTooltip", "Set the units descriptor for this property."),
-							BuildGenericSpecifierSubMenu(Metadata::Specifiers::UnitsName),
-							/*bInOpenSubMenuOnClick=*/true);
-					}
+					MenuBuilder.AddSubMenu(
+						FText::FromName(Metadata::Specifiers::UIMaxName),
+						LOCTEXT("DropDownMenuSubMenuUIMaxTooltip", "Set the max UI value for this property."),
+						BuildGenericSpecifierSubMenu(Metadata::Specifiers::UIMaxName, /*bIsNumeric=*/true),
+						/*bInOpenSubMenuOnClick=*/true);
+				}
 
-					MenuBuilder.EndSection();
+				if (Metadata::Specifiers::SupportsSettingUnits(*PropertyDesc))
+				{
+					MenuBuilder.AddSubMenu(
+						FText::FromName(Metadata::Specifiers::UnitsName),
+						LOCTEXT("DropDownMenuSubMenuUnitsTooltip", "Set the units descriptor for this property."),
+						BuildGenericSpecifierSubMenu(Metadata::Specifiers::UnitsName),
+						/*bInOpenSubMenuOnClick=*/true);
 				}
+
+				MenuBuilder.EndSection();
 			}
 
 			/*** DROP-DOWN ARROW MENU ***/
@@ -1961,92 +1696,88 @@ void FPropertyBagInstanceDataDetails::OnChildRowAdded(IDetailPropertyRow& ChildR
 				];
 		}
 
-		/** DRAG AND DROP HANDLER 
-		 *  We only add drag & drop for reordering if there is no hierarchy root. */
+		/*** DRAG AND DROP HANDLER ***/
 		if (EnumHasAnyFlags(ChildRowFeatures, EPropertyBagChildRowFeatures::DragAndDrop) && PropertyDesc != nullptr)
 		{
 			TSharedPtr<FPropertyBagDetailsDragDropHandler> DragDropHandler = MakeShared<FPropertyBagDetailsDragDropHandler>(*PropertyDesc);
 
-			// If we support drag & drop but a hierarchy root exists, we allow dragging, but not dropping on our property rows
-			if (!HierarchyData.HierarchyRoot.IsValid())
-			{
-				// Can accept drag and drop check for if this is a valid drop
-				DragDropHandler->BindCanAcceptDragDrop(
-					FCanAcceptPropertyBagDetailsRowDropOp::CreateLambda(
-						[PropertyDesc](const TSharedPtr<FPropertyBagDetailsDragDropOp>& DropOp, EItemDropZone DropZone) -> TOptional<EItemDropZone>
+			// Can accept drag and drop check for if this is a valid drop
+			DragDropHandler->BindCanAcceptDragDrop(
+				FCanAcceptPropertyBagDetailsRowDropOp::CreateLambda(
+					[PropertyDesc](const TSharedPtr<FPropertyBagDetailsDragDropOp>& DropOp, EItemDropZone DropZone) -> TOptional<EItemDropZone>
+					{
+						if (!PropertyDesc || !DropOp.IsValid())
 						{
-							if (!PropertyDesc || !DropOp.IsValid())
-							{
-								DropOp->SetDecoration(EPropertyBagDropState::Invalid);
-								return TOptional<EItemDropZone>();
-							}
+							DropOp->SetDecoration(EPropertyBagDropState::Invalid);
+							return TOptional<EItemDropZone>();
+						}
 
-							if (DropZone == EItemDropZone::OntoItem && PropertyDesc->ID != DropOp->PropertyDesc.ID)
-							{
-								DropOp->SetDecoration(EPropertyBagDropState::Invalid);
-								return TOptional<EItemDropZone>();
-							}
+						if (DropZone == EItemDropZone::OntoItem && PropertyDesc->ID != DropOp->PropertyDesc.ID)
+						{
+							DropOp->SetDecoration(EPropertyBagDropState::Invalid);
+							return TOptional<EItemDropZone>();
+						}
 
-							// No effect to drop in these cases. Either source == target, or moving source above/below target puts source in same location.
-							if (*PropertyDesc == DropOp->PropertyDesc
-								|| (DropZone == EItemDropZone::AboveItem && DropOp->PropertyDesc.GetCachedIndex() == PropertyDesc->GetCachedIndex() - 1)
-								|| (DropZone == EItemDropZone::BelowItem && DropOp->PropertyDesc.GetCachedIndex() == PropertyDesc->GetCachedIndex() + 1))
-							{
-								DropOp->SetDecoration(EPropertyBagDropState::SourceIsTarget);
-								return TOptional<EItemDropZone>();
-							}
+						// No effect to drop in these cases. Either source == target, or moving source above/below target puts source in same location.
+						if (*PropertyDesc == DropOp->PropertyDesc
+							|| (DropZone == EItemDropZone::AboveItem && DropOp->PropertyDesc.GetCachedIndex() == PropertyDesc->GetCachedIndex() - 1)
+							|| (DropZone == EItemDropZone::BelowItem && DropOp->PropertyDesc.GetCachedIndex() == PropertyDesc->GetCachedIndex() + 1))
+						{
+							DropOp->SetDecoration(EPropertyBagDropState::SourceIsTarget);
+							return TOptional<EItemDropZone>();
+						}
 
-							DropOp->SetDecoration(EPropertyBagDropState::Valid);
-							return DropZone;
-						}));
+						DropOp->SetDecoration(EPropertyBagDropState::Valid);
+						return DropZone;
+					}));
 
-				DragDropHandler->BindOnHandleDragDrop(
-					FOnPropertyBagDetailsRowDropOp::CreateLambda(
-						[WeakSelf, PropertyDesc = *PropertyDesc, PropertyBagHandle = BagStructProperty, PropUtils = PropUtils](const FPropertyBagPropertyDesc& DroppedPropertyDesc, EItemDropZone DropZone) -> FReply
+			DragDropHandler->BindOnHandleDragDrop(
+				FOnPropertyBagDetailsRowDropOp::CreateLambda(
+					[WeakSelf, PropertyDesc = *PropertyDesc, BagStructProperty = BagStructProperty, PropUtils = PropUtils](const FPropertyBagPropertyDesc& DroppedPropertyDesc, EItemDropZone DropZone) -> FReply
+					{
+						using namespace UE::StructUtils;
+
+						const TSharedPtr<FPropertyBagInstanceDataDetails> DetailsSP = WeakSelf.Pin();
+						const UPropertyBag* ChildBagStruct = DetailsSP ? GetCommonBagStruct(DetailsSP->BagStructProperty) : nullptr;
+						// Validate these properties are still part of the bag.
+						if (!ChildBagStruct
+							|| !ChildBagStruct->FindPropertyDescByProperty(PropertyDesc.CachedProperty)
+							|| !ChildBagStruct->FindPropertyDescByProperty(DroppedPropertyDesc.CachedProperty))
 						{
-							using namespace UE::StructUtils;
-
-							const TSharedPtr<FPropertyBagInstanceDataDetails> DetailsSP = WeakSelf.Pin();
-							const UPropertyBag* ChildBagStruct = DetailsSP ? GetCommonBagStruct(DetailsSP->BagStructProperty) : nullptr;
-							// Validate these properties are still part of the bag.
-							if (!ChildBagStruct
-								|| !ChildBagStruct->FindPropertyDescByProperty(PropertyDesc.CachedProperty)
-								|| !ChildBagStruct->FindPropertyDescByProperty(DroppedPropertyDesc.CachedProperty))
-							{
-								return FReply::Unhandled();
-							}
+							return FReply::Unhandled();
+						}
 
-							EPropertyBagAlterationResult Result = EPropertyBagAlterationResult::InternalError;
+						EPropertyBagAlterationResult Result = EPropertyBagAlterationResult::InternalError;
 
-							DetailsSP->BagStructProperty->EnumerateRawData([DropZone, &PropertyDesc, &DroppedPropertyDesc, &Result](void* RawData, const int32 /*DataIndex*/, const int32 /*NumData*/)
+						DetailsSP->BagStructProperty->EnumerateRawData([DropZone, &PropertyDesc, &DroppedPropertyDesc, &Result](void* RawData, const int32 /*DataIndex*/, const int32 /*NumData*/)
+						{
+							if (FInstancedPropertyBag* PropertyBag = RawData ? static_cast<FInstancedPropertyBag*>(RawData) : nullptr)
 							{
-								if (FInstancedPropertyBag* PropertyBag = RawData ? static_cast<FInstancedPropertyBag*>(RawData) : nullptr)
-								{
-									Result = PropertyBag->ReorderProperty(DroppedPropertyDesc.Name, PropertyDesc.Name, DropZone == EItemDropZone::AboveItem);
-								}
+								Result = PropertyBag->ReorderProperty(DroppedPropertyDesc.Name, PropertyDesc.Name, DropZone == EItemDropZone::AboveItem);
+							}
 
-								return true;
-							});
+							return true;
+						});
 
-							if (Result == EPropertyBagAlterationResult::Success)
-							{
-								ApplyChangesToSinglePropertyDesc(
-									LOCTEXT("DragDropReorderProperties", "Reordered properties"),
-									DroppedPropertyDesc,
-									PropertyBagHandle,
-									[PropertyDesc](FPropertyBagPropertyDesc& Desc)
-										{
-											Metadata::SetCategory(Desc, Metadata::GetCategory(PropertyDesc));
-										});
+						if (Result == EPropertyBagAlterationResult::Success)
+						{
+							ApplyChangesToSinglePropertyDesc(
+								LOCTEXT("DragDropReorderProperties", "Reordered properties"),
+								DroppedPropertyDesc,
+								BagStructProperty,
+								PropUtils,
+								[PropertyDesc](FPropertyBagPropertyDesc& Desc)
+								{
+									Metadata::SetCategory(Desc, Metadata::GetCategory(PropertyDesc));
+								});
 
-								return FReply::Handled();
-							}
-							else
-							{
-								return FReply::Unhandled();
-							}
-						}));
-			}
+							return FReply::Handled();
+						}
+						else
+						{
+							return FReply::Unhandled();
+						}
+					}));
 
 			// Bind the drag and drop handler for receiving.
 			ChildRow.DragDropHandler(DragDropHandler);
@@ -2082,9 +1813,8 @@ void FPropertyBagInstanceDataDetails::OnChildRowAdded(IDetailPropertyRow& ChildR
 					]
 				];
 		}
-		else
+		else // Update the name widget with our new property details composition.
 		{
-			// Update the name widget with our new property details composition.
 			NameWidget = PropertyDetailsWidget;
 		}
 	}
@@ -2269,48 +1999,34 @@ void FPropertyBagInstanceDataDetails::ResetToDefault(TSharedPtr<IPropertyHandle>
 	ChildPropertyHandle->NotifyFinishedChangingProperties();
 }
 
-EVisibility FPropertyBagInstanceDataDetails::ShouldShowHierarchyProperty(TSharedPtr<IPropertyHandle> ChildPropertyHandle) const
-{
-	if (bShouldShowAdvanced)
-	{
-		return EVisibility::Visible;
-	}
-	
-	if (FProperty* MetaDataProperty = ChildPropertyHandle->GetMetaDataProperty())
-	{
-		return MetaDataProperty->HasAnyPropertyFlags(CPF_AdvancedDisplay) ? EVisibility::Collapsed : EVisibility::Visible;
-	}
-	
-	return EVisibility::Visible;
-}
-
 TSharedRef<SWidget> FPropertyBagInstanceDataDetails::OnPropertyNameContent(TSharedPtr<IPropertyHandle> ChildPropertyHandle, TSharedPtr<SInlineEditableTextBlock> InlineWidget) const
 {
 	constexpr bool bInShouldCloseWindowAfterMenuSelection = true;
 	FMenuBuilder MenuBuilder(bInShouldCloseWindowAfterMenuSelection, nullptr);
 
-	auto MoveProperty = [PropertyBagHandle = BagStructProperty, PropUtils = PropUtils, ChildPropertyHandle](const int32 Delta)
+	auto MoveProperty = [BagStructProperty = BagStructProperty, PropUtils = PropUtils, ChildPropertyHandle](const int32 Delta)
 	{
-		if (!PropertyBagHandle || !PropertyBagHandle->IsValidHandle() || !ChildPropertyHandle || !ChildPropertyHandle->IsValidHandle())
+		if (!BagStructProperty || !BagStructProperty->IsValidHandle() || !ChildPropertyHandle || !ChildPropertyHandle->IsValidHandle())
 		{
 			return;
 		}
 
 		UE::StructUtils::ApplyChangesToPropertyDescs(
-			LOCTEXT("OnPropertyMoved", "Move Property"), PropertyBagHandle, [&ChildPropertyHandle, &Delta](TArray<FPropertyBagPropertyDesc>& PropertyDescs)
+		LOCTEXT("OnPropertyMoved", "Move Property"), BagStructProperty, PropUtils,
+		[&ChildPropertyHandle, &Delta](TArray<FPropertyBagPropertyDesc>& PropertyDescs)
+		{
+			// Move
+			if (PropertyDescs.Num() > 1)
+			{
+				const FProperty* Property = ChildPropertyHandle ? ChildPropertyHandle->GetProperty() : nullptr;
+				const int32 PropertyIndex = PropertyDescs.IndexOfByPredicate([Property](const FPropertyBagPropertyDesc& Desc){ return Desc.CachedProperty == Property; });
+				if (PropertyIndex != INDEX_NONE)
 				{
-					// Move
-					if (PropertyDescs.Num() > 1)
-					{
-						const FProperty* Property = ChildPropertyHandle ? ChildPropertyHandle->GetProperty() : nullptr;
-						const int32 PropertyIndex = PropertyDescs.IndexOfByPredicate([Property](const FPropertyBagPropertyDesc& Desc){ return Desc.CachedProperty == Property; });
-						if (PropertyIndex != INDEX_NONE)
-						{
-							const int32 NewPropertyIndex = FMath::Clamp(PropertyIndex + Delta, 0, PropertyDescs.Num() - 1);
-							PropertyDescs.Swap(PropertyIndex, NewPropertyIndex);
-						}
-					}
-				});
+					const int32 NewPropertyIndex = FMath::Clamp(PropertyIndex + Delta, 0, PropertyDescs.Num() - 1);
+					PropertyDescs.Swap(PropertyIndex, NewPropertyIndex);
+				}
+			}
+		});
 	};
 
 	MenuBuilder.AddWidget(
@@ -2323,7 +2039,7 @@ TSharedRef<SWidget> FPropertyBagInstanceDataDetails::OnPropertyNameContent(TShar
 				PropUtils,
 				SPinTypeSelector::ESelectorType::Full,
 				bAllowContainers,
-				PropertyBagSchema->GetClass())
+				PropertyBagSchemaClass)
 		],
 		FText::GetEmpty());
 
@@ -2365,259 +2081,6 @@ TSharedRef<SWidget> FPropertyBagInstanceDataDetails::OnPropertyNameContent(TShar
 	return MenuBuilder.MakeWidget();
 }
 
-void FPropertyBagInstanceDataDetails::GenerateHierarchySectionRow(IDetailChildrenBuilder& ChildrenBuilder)
-{
-	if (!HierarchyData.HierarchyRoot.IsValid())
-	{
-		return;
-	}
-	
-	TArray<const UHierarchySection*> OrderedSections;
-	TMap<const UHierarchySection*, TArray<const UHierarchyCategory*>> SectionCategoryMap;
-
-	UHierarchyRoot* Root = HierarchyData.HierarchyRoot.Get();
-	for(const UHierarchyElement* Child : Root->GetChildren())
-	{
-		if(const UHierarchyCategory* HierarchyCategory = Cast<UHierarchyCategory>(Child))
-		{
-			const FDataHierarchyElementMetaData_SectionAssociation* SectionAssociation = HierarchyCategory->FindMetaDataOfType<FDataHierarchyElementMetaData_SectionAssociation>();
-			if(SectionAssociation && SectionAssociation->Section.IsValid())
-			{
-				SectionCategoryMap.FindOrAdd(SectionAssociation->Section.Get()).Add(HierarchyCategory);	
-			}
-		}
-	}
-
-	// we create a custom row for user param sections here, but only if we have at least one section specified
-	if(SectionCategoryMap.Num() > 0)
-	{
-		TSharedRef<SWrapBox> SectionsBox = SNew(SWrapBox).UseAllottedSize(true);
-
-		// maps don't guarantee order, so we iterate over the original section data instead
-		for(const TObjectPtr<UHierarchySection>& Section : Root->GetSectionData())
-		{
-			// if the map doesn't contain at the entry, it means there are 0 categories for that section, so we skip it
-			if(!SectionCategoryMap.Contains(Section))
-			{
-				continue;
-			}
-			
-			TArray<const UHierarchyCategory*>& HierarchyCategories = SectionCategoryMap[Section];
-			
-			if(HierarchyCategories.Num() > 0)
-			{
-				// we only want to add a section if any of its categories contain at least one parameter
-				bool bDoesUserParamExist = false;
-				for(const UHierarchyCategory* Category : HierarchyCategories)
-				{
-					bDoesUserParamExist |= Category->DoesOneChildExist<UPropertyBagHierarchyProperty>(true);
-
-					if(bDoesUserParamExist)
-					{
-						break;
-					}
-				}
-
-				if(bDoesUserParamExist)
-				{
-					SectionsBox->AddSlot()
-					.Padding(FMargin(2.f))
-					[
-						SNew(SBox)
-						.Padding(FMargin(0.f, 4.f, 0.f, 0.f))
-						[
-							SNew(SCheckBox)
-							.Style(FAppStyle::Get(), "DetailsView.SectionButton")
-							.OnCheckStateChanged_Lambda([this, Section](ECheckBoxState NewState)
-							{
-								HierarchyData.ActiveHierarchySection = Section;
-								OnRebuildChildren.ExecuteIfBound();
-							})
-							.IsChecked_Lambda([this, Section]()
-							{
-								return HierarchyData.ActiveHierarchySection == Section ? ECheckBoxState::Checked : ECheckBoxState::Unchecked;
-							})
-							[
-								SNew(STextBlock)
-								.TextStyle(FAppStyle::Get(), "SmallText")
-								.Text_UObject(Section.Get(), &UHierarchySection::GetSectionNameAsText)
-								.ToolTipText_UObject(Section.Get(), &UHierarchySection::GetTooltip)
-							]
-						]
-					];
-				}
-			}
-		}
-
-		// if we have at least one custom section, we add a default "All" section
-		if(SectionsBox->GetChildren()->Num() > 0)
-		{
-			SectionsBox->AddSlot()
-			.Padding(FMargin(2.f))
-			[
-				SNew(SBox)
-				.Padding(FMargin(0.f, 4.f, 0.f, 0.f))
-				[
-					SNew(SCheckBox)
-					.Style(FAppStyle::Get(), "DetailsView.SectionButton")
-					.OnCheckStateChanged_Lambda([this](ECheckBoxState NewState)
-					{
-						HierarchyData.ActiveHierarchySection = nullptr;
-						OnRebuildChildren.ExecuteIfBound();
-					})
-					.IsChecked_Lambda([this]()
-					{
-						return HierarchyData.ActiveHierarchySection == nullptr ? ECheckBoxState::Checked : ECheckBoxState::Unchecked;
-					})
-					[
-						SNew(STextBlock)
-						.TextStyle(FAppStyle::Get(), "SmallText")
-						.Text(FText::FromString("All"))
-					]
-				]
-			];
-			
-			FDetailWidgetRow& Row = ChildrenBuilder.AddCustomRow(FText::FromString("Sections"));
-			Row.RowTag(FName("Sections"));
-			
-			Row.WholeRowContent()
-			[
-				SectionsBox
-			];
-		}
-	}
-}
-
-void FPropertyBagInstanceDataDetails::GenerateHierarchyRootRows(IDetailChildrenBuilder& ChildrenBuilder)
-{
-	if (!HierarchyData.HierarchyRoot.IsValid())
-	{
-		return;
-	}
-	
-	// Then, we generate root level categories...
-	TArray<const UPropertyBagHierarchyCategory*> Categories;
-	HierarchyData.HierarchyRoot->GetChildrenOfType(Categories, false);
-		
-	for(const UPropertyBagHierarchyCategory* Category : Categories)
-	{
-		const UHierarchySection* Section = nullptr;
-		if(const FDataHierarchyElementMetaData_SectionAssociation* SectionAssociation = Category->FindMetaDataOfType<FDataHierarchyElementMetaData_SectionAssociation>())
-		{
-			if(SectionAssociation->Section.IsValid())
-			{
-				Section = SectionAssociation->Section.Get();
-			}
-		}
-			
-		// ... but skip categories where the section doesn't match or categories that don't actually contain a property
-		if((HierarchyData.ActiveHierarchySection == nullptr || Section == HierarchyData.ActiveHierarchySection) && Category->DoesOneChildExist<UPropertyBagHierarchyProperty>(true))
-		{
-			ChildrenBuilder.AddCustomBuilder(MakeShared<FHierarchyPropertyCategoryBuilder>(*Category, SharedThis(this)));
-		}
-	}
-		
-	// Then we add root level properties
-	TArray<const UPropertyBagHierarchyProperty*> HierarchyPropertiesUnderRoot;
-	HierarchyData.HierarchyRoot->GetChildrenOfType(HierarchyPropertiesUnderRoot, false);
-	for(const UPropertyBagHierarchyProperty* HierarchyProperty : HierarchyPropertiesUnderRoot)
-	{
-		// we only want to display parameters within the root in the all section
-		if(HierarchyData.ActiveHierarchySection == nullptr)
-		{
-			if (HierarchyData.HierarchyPropertyHandleMap.Contains(HierarchyProperty))
-			{
-				IDetailPropertyRow& PropertyRow = ChildrenBuilder.AddProperty(HierarchyData.HierarchyPropertyHandleMap[HierarchyProperty].ToSharedRef());
-				PropertyRow.Visibility(TAttribute<EVisibility>::CreateSP(this, &FPropertyBagInstanceDataDetails::ShouldShowHierarchyProperty, PropertyRow.GetPropertyHandle()));
-				OnChildRowAdded(PropertyRow);
-			}
-		}
-	}
-}
-
-void FPropertyBagInstanceDataDetails::GenerateHierarchyLeftoverPropertyRows(IDetailChildrenBuilder& ChildrenBuilder, const TArray<TSharedPtr<IPropertyHandle>>& AllChildPropertyHandles)
-{
-	if (!HierarchyData.HierarchyRoot.IsValid())
-	{
-		return;
-	}
-	
-	// Then we add all parameters that haven't been setup in the hierarchy at all, if the active section is set to "All"
-	if (HierarchyData.ActiveHierarchySection == nullptr)
-	{
-		TArray<const UPropertyBagHierarchyProperty*> AllHierarchyProperties;
-		HierarchyData.HierarchyRoot->GetChildrenOfType(AllHierarchyProperties, true);
-		
-		TArray<TSharedPtr<IPropertyHandle>> LeftoverProperties = AllChildPropertyHandles;
-		TArray<TSharedPtr<IPropertyHandle>> PropertyHandlesInHierarchy;
-		HierarchyData.HierarchyPropertyHandleMap.GenerateValueArray(PropertyHandlesInHierarchy);
-			
-		LeftoverProperties.RemoveAll([PropertyHandlesInHierarchy](TSharedPtr<IPropertyHandle> RemovalCandidate)
-		{
-			return PropertyHandlesInHierarchy.Contains(RemovalCandidate);
-		});
-			
-		for (TSharedPtr<IPropertyHandle> LeftoverProperty : LeftoverProperties)
-		{
-			IDetailPropertyRow& PropertyRow = ChildrenBuilder.AddProperty(LeftoverProperty.ToSharedRef());
-			PropertyRow.Visibility(TAttribute<EVisibility>::CreateSP(this, &FPropertyBagInstanceDataDetails::ShouldShowHierarchyProperty, PropertyRow.GetPropertyHandle()));
-			OnChildRowAdded(PropertyRow);
-		}
-	}
-}
-
-const FSlateBrush* FPropertyBagInstanceDataDetails::GetAdvancedImage() const
-{	
-	FName ResourceName;
-	if (bShouldShowAdvanced)
-	{
-		if (ShowAdvancedButton->IsHovered())
-		{
-			static const FName ExpandedHoveredName = "DetailsView.PulldownArrow.Up.Hovered";
-			ResourceName = ExpandedHoveredName;
-		}
-		else
-		{
-			static const FName ExpandedName = "DetailsView.PulldownArrow.Up";
-			ResourceName = ExpandedName;
-		}
-	}
-	else
-	{
-		if (ShowAdvancedButton->IsHovered())
-		{
-			static const FName CollapsedHoveredName = "DetailsView.PulldownArrow.Down.Hovered";
-			ResourceName = CollapsedHoveredName;
-		}
-		else
-		{
-			static const FName CollapsedName = "DetailsView.PulldownArrow.Down";
-			ResourceName = CollapsedName;
-		}
-	}
-	
-	return FAppStyle::GetBrush(ResourceName);
-}
-
-EVisibility FPropertyBagInstanceDataDetails::GetAdvancedButtonVisibility() const
-{
-	if (HierarchyData.HierarchyRoot.IsValid())
-	{
-		TArray<UPropertyBagHierarchyProperty*> AllHierarchyProperties;
-		HierarchyData.HierarchyRoot->GetChildrenOfType<UPropertyBagHierarchyProperty>(AllHierarchyProperties, true);
-		
-		for (UPropertyBagHierarchyProperty* Property : AllHierarchyProperties)
-		{
-			if (Property->bAdvanced)
-			{
-				return EVisibility::Visible;
-			}
-		}
-	}
-	
-	return EVisibility::Collapsed;
-}
-
 ////////////////////////////////////
 
 TSharedRef<IPropertyTypeCustomization> FPropertyBagDetails::MakeInstance()
@@ -2632,8 +2095,6 @@ void FPropertyBagDetails::CustomizeHeader(TSharedRef<IPropertyHandle> StructProp
 	StructProperty = StructPropertyHandle;
 	check(StructProperty);
 
-	PropertyBagSchemaCDO = UE::StructUtils::ExtractPropertyBagSchemaCDO(StructPropertyHandle);
-	
 	if (const FProperty* MetaDataProperty = StructProperty->GetMetaDataProperty())
 	{
 		PRAGMA_DISABLE_DEPRECATION_WARNINGS
@@ -2643,7 +2104,9 @@ void FPropertyBagDetails::CustomizeHeader(TSharedRef<IPropertyHandle> StructProp
 		bAllowContainers = MetaDataProperty->HasMetaData(UE::StructUtils::Metadata::AllowContainersName)
 			? MetaDataProperty->GetBoolMetaData(UE::StructUtils::Metadata::AllowContainersName)
 			: true;
-		
+
+		PropertyBagSchemaClass = FindObject<UClass>(nullptr, MetaDataProperty->GetMetaData(UE::StructUtils::Metadata::PropertyBagSchemaClassName));
+
 		if (MetaDataProperty->HasMetaData(UE::StructUtils::Metadata::DefaultTypeName))
 		{
 			if (UEnum* Enum = StaticEnum<EPropertyBagPropertyType>())
@@ -2683,12 +2146,13 @@ void FPropertyBagDetails::CustomizeHeader(TSharedRef<IPropertyHandle> StructProp
 		}
 	}
 
-	TSharedPtr<SHorizontalBox> ValueWidget = SNew(SHorizontalBox);
+	TSharedPtr<SWidget> ValueWidget = SNullWidget::NullWidget;
 	PRAGMA_DISABLE_DEPRECATION_WARNINGS
 	if (!bFixedLayout && ChildRowFeatures != EPropertyBagChildRowFeatures::Fixed)
 	PRAGMA_ENABLE_DEPRECATION_WARNINGS
 	{
-		ValueWidget->AddSlot()
+		SAssignNew(ValueWidget, SHorizontalBox)
+		+ SHorizontalBox::Slot()
 		.AutoWidth()
 		[
 			MakeAddPropertyWidget(StructProperty, PropUtils, DefaultType).ToSharedRef()
@@ -2710,8 +2174,6 @@ void FPropertyBagDetails::CustomizeHeader(TSharedRef<IPropertyHandle> StructProp
 
 void FPropertyBagDetails::CustomizeChildren(TSharedRef<IPropertyHandle> StructPropertyHandle, IDetailChildrenBuilder& StructBuilder, IPropertyTypeCustomizationUtils& StructCustomizationUtils)
 {
-	FPropertyBagInstanceDataDetails::FConstructParams::FHierarchyParams HierarchyParams;
-	HierarchyParams.OnNavigateToHierarchyPropertyRequested = FPropertyBagInstanceDataDetails::FPropertyBagPropertyDelegate::CreateSP(this, &FPropertyBagDetails::NavigateToHierarchyProperty); 
 	FPropertyBagInstanceDataDetails::FConstructParams Params
 	{
 		.BagStructProperty = StructProperty,
@@ -2720,7 +2182,7 @@ void FPropertyBagDetails::CustomizeChildren(TSharedRef<IPropertyHandle> StructPr
 		PRAGMA_DISABLE_DEPRECATION_WARNINGS
 		.ChildRowFeatures = bFixedLayout ? EPropertyBagChildRowFeatures::Fixed : ChildRowFeatures,
 		PRAGMA_ENABLE_DEPRECATION_WARNINGS
-		.HierarchyParams = HierarchyParams
+		.PropertyBagSchemaClass = PropertyBagSchemaClass
 	};
 
 	const UPropertyBagDetailsExtension* Extension = GetDefault<UPropertyBagDetailsExtension>();
@@ -2743,340 +2205,79 @@ void FPropertyBagDetails::CustomizeChildren(TSharedRef<IPropertyHandle> StructPr
 	StructBuilder.AddCustomBuilder(InstanceDetails);
 }
 
-FReply FPropertyBagDetails::SummonHierarchyEditor(TSharedPtr<IPropertyHandle> InPropertyBagHandle, TSharedPtr<IPropertyUtilities> InPropertyUtilities)
-{
-	if (!InPropertyBagHandle.IsValid() || !InPropertyBagHandle->IsValidHandle() || !InPropertyUtilities.IsValid())
-	{
-		return FReply::Unhandled();
-	}
-	
-	UPropertyBagHierarchyRoot* HierarchyRoot = UE::StructUtils::ExtractHierarchyRoot(InPropertyBagHandle.ToSharedRef());
-	
-	if (!HierarchyRoot)
-	{
-		return FReply::Unhandled();
-	}
-	
-	bool bNewWindow = false;
-	
-	TSharedPtr<SWindow> Window;
-
-	TArray<UObject*> OuterObjects;
-	InPropertyBagHandle->GetOuterObjects(OuterObjects);
-	
-	if (OuterObjects.Num() != 1)
-	{
-		return FReply::Unhandled();
-	}
-	
-	FString OwnerName = OuterObjects[0]->GetName();
-	
-	FText Title = FText::FormatOrdered(INVTEXT("{0} - {1}"), FText::FromString(OwnerName), InPropertyBagHandle->GetPropertyDisplayName());
-	// We reuse the existing window if possible
-	if (!PropertyBagHierarchyEditorWindow.IsValid())
-	{
-		Window = SNew(SWindow)
-		.Title(Title)
-		.IsTopmostWindow(true)
-		.ClientSize(FVector2f(750.f, 450.f));
-		
-		PropertyBagHierarchyEditorWindow = Window;
-		bNewWindow = true;
-	}
-	else
-	{
-		Window = PropertyBagHierarchyEditorWindow.Pin();
-	}
-	
-	Window->SetTitle(Title);
-	
-	TSharedPtr<SPropertyBagHierarchyEditor> NewPropertyBagHierarchyEditor = SNew(SPropertyBagHierarchyEditor, InPropertyBagHandle.ToSharedRef(), InPropertyUtilities.ToSharedRef())
-	.OnCloseRequested_Static(&FPropertyBagDetails::ClosePropertyBagHierarchyEditorWindow);
-	
-	PropertyBagHierarchyEditor = NewPropertyBagHierarchyEditor;
-	Window->SetContent(NewPropertyBagHierarchyEditor.ToSharedRef());
-	
-	if (bNewWindow)
-	{
-		FSlateApplication::Get().AddWindow(Window.ToSharedRef(), true);
-	}
-	else
-	{
-		FWindowDrawAttentionParameters Parameters;
-		Parameters.RequestType = EWindowDrawAttentionRequestType::Stop;
-		Window->DrawAttention(Parameters);
-	}
-	
-	return FReply::Handled();
-}
-
-void FPropertyBagDetails::NavigateToHierarchyProperty(const FPropertyBagPropertyDesc& Desc)
-{
-	if (!Desc.ID.IsValid())
-	{
-		return;
-	}
-	
-	if (!PropertyBagHierarchyEditor.IsValid() || !PropertyBagHierarchyEditor.Pin()->GetPropertyBagHandle().IsValid() || !PropertyBagHierarchyEditor.Pin()->GetPropertyBagHandle()->IsSamePropertyNode(StructProperty))
-	{
-		SummonHierarchyEditor(StructProperty.ToSharedRef(), PropUtils.ToSharedRef());
-	}
-	
-	if (PropertyBagHierarchyEditor.IsValid())
-	{
-		PropertyBagHierarchyEditor.Pin()->NavigateToProperty(Desc, true);
-	}
-}
-
-void FPropertyBagDetails::ClosePropertyBagHierarchyEditorWindow()
-{
-	if (PropertyBagHierarchyEditorWindow.IsValid())
-	{
-		PropertyBagHierarchyEditorWindow.Pin()->RequestDestroyWindow();
-	}
-}
-
 TSharedPtr<SWidget> FPropertyBagDetails::MakeAddPropertyWidget(TSharedPtr<IPropertyHandle> InStructProperty, TSharedPtr<IPropertyUtilities> InPropUtils, EPropertyBagPropertyType DefaultType, const FSlateColor IconColor)
 {
 	const FText RemoveAllPropertiesText = LOCTEXT("OnRemoveAllProperties", "Remove all Properties");
 
-	TSharedRef<SHorizontalBox> Container = SNew(SHorizontalBox)
-	+SHorizontalBox::Slot()
-	.AutoWidth()
-	[
-		SNew(SButton)
-		.VAlign(VAlign_Center)
-		.HAlign(HAlign_Center)
-		.ButtonStyle(FAppStyle::Get(), "SimpleButton")
-		.ToolTipText(LOCTEXT("AddProperty_Tooltip", "Add new property"))
-		.OnClicked_Lambda([StructProperty = InStructProperty, PropUtils = InPropUtils, DefaultType]()
-		{
-			constexpr int32 MaxIterations = 100;
-			FName NewName(TEXT("NewProperty"));
-			int32 Number = 1;
-			while (!UE::StructUtils::Private::IsUniqueName(NewName, FName(), StructProperty) && Number < MaxIterations)
-			{
-				Number++;
-				NewName.SetNumber(Number);
-			}
-			if (Number == MaxIterations)
+	return SNew(SHorizontalBox)
+		+SHorizontalBox::Slot()
+		.AutoWidth()
+		[
+			SNew(SButton)
+			.VAlign(VAlign_Center)
+			.HAlign(HAlign_Center)
+			.ButtonStyle(FAppStyle::Get(), "SimpleButton")
+			.ToolTipText(LOCTEXT("AddProperty_Tooltip", "Add new property"))
+			.OnClicked_Lambda([StructProperty = InStructProperty, PropUtils = InPropUtils, DefaultType]()
 			{
-				return FReply::Handled();
-			}
-
-			UE::StructUtils::ApplyChangesToPropertyDescs(
-				LOCTEXT("OnPropertyAdded", "Add Property"), StructProperty,
-				[&NewName, DefaultType](TArray<FPropertyBagPropertyDesc>& PropertyDescs)
+				constexpr int32 MaxIterations = 100;
+				FName NewName(TEXT("NewProperty"));
+				int32 Number = 1;
+				while (!UE::StructUtils::Private::IsUniqueName(NewName, FName(), StructProperty) && Number < MaxIterations)
 				{
-					PropertyDescs.Emplace(NewName, DefaultType);
-				});
+					Number++;
+					NewName.SetNumber(Number);
+				}
+				if (Number == MaxIterations)
+				{
+					return FReply::Handled();
+				}
 
-			return FReply::Handled();
+				UE::StructUtils::ApplyChangesToPropertyDescs(
+					LOCTEXT("OnPropertyAdded", "Add Property"), StructProperty, PropUtils,
+					[&NewName, DefaultType](TArray<FPropertyBagPropertyDesc>& PropertyDescs)
+					{
+						PropertyDescs.Emplace(NewName, DefaultType);
+					});
 
-		})
-		.Content()
-		[
-			SNew(SImage)
-			.Image(FAppStyle::GetBrush("Icons.PlusCircle"))
-			.ColorAndOpacity(IconColor)
-		]
-	]
-	+ SHorizontalBox::Slot()
-	.AutoWidth()
-	.Padding(FMargin(0.f, 0.f, 4.f, 0.f))
-	[
-		PropertyCustomizationHelpers::MakeEmptyButton(
-			FSimpleDelegate::CreateLambda([StructProperty = InStructProperty, PropUtils = InPropUtils, RemoveAllPropertiesText]()
-				{
-					UE::StructUtils::ApplyChangesToPropertyDescs(
-						RemoveAllPropertiesText,
-						StructProperty,
-						[](TArray<FPropertyBagPropertyDesc>& PropertyDescs)
-						{
-							PropertyDescs.Empty();
-						});
-				}),
-			RemoveAllPropertiesText
-			)
-	];
-	
-	// We attempt to retrieve a hierarchy root via schema, and display the Edit Hierarchy button if successful
-	if (UPropertyBagHierarchyRoot* Root = UE::StructUtils::ExtractHierarchyRoot(InStructProperty.ToSharedRef()))
-	{		
-		Container->AddSlot()
-		.AutoWidth()
-		.HAlign(HAlign_Left)
-		.VAlign(VAlign_Center)
-		[
-			SNew(SButton)
-			.ToolTipText(LOCTEXT("EditHierarchy_Tooltip", "Summon Hierarchy Editor to edit categories and metadata."))
-			.ButtonStyle(&FAppStyle::GetWidgetStyle<FButtonStyle>("HoverHintOnly"))
-			.OnClicked_Static(&SummonHierarchyEditor, InStructProperty, InPropUtils)
+				return FReply::Handled();
+
+			})
+			.Content()
 			[
 				SNew(SImage)
-				.Image(FAppStyle::GetBrush("Icons.EditHierarchy"))
+				.Image(FAppStyle::GetBrush("Icons.PlusCircle"))
 				.ColorAndOpacity(IconColor)
 			]
+		]
+		+ SHorizontalBox::Slot()
+		.AutoWidth()
+		.Padding(FMargin(0.f, 0.f, 4.f, 0.f))
+		[
+			PropertyCustomizationHelpers::MakeEmptyButton(
+				FSimpleDelegate::CreateLambda([StructProperty = InStructProperty, PropUtils = InPropUtils, RemoveAllPropertiesText]()
+					{
+						UE::StructUtils::ApplyChangesToPropertyDescs(
+							RemoveAllPropertiesText,
+							StructProperty,
+							PropUtils,
+							[](TArray<FPropertyBagPropertyDesc>& PropertyDescs)
+							{
+								PropertyDescs.Empty();
+							});
+					}),
+				RemoveAllPropertiesText
+				)
 		];
-	}
-	
-	return Container;
-}
-
-////////////////////////////////////
-bool UPropertyBagSchema::SupportsPinTypeContainer(TWeakPtr<const FEdGraphSchemaAction> SchemaAction, const FEdGraphPinType& PinType, const EPinContainerType& ContainerType) const
-{
-	return ContainerType == EPinContainerType::None || ContainerType == EPinContainerType::Array || ContainerType == EPinContainerType::Set || ContainerType == EPinContainerType::Map;
-}
 
-UPropertyBagHierarchyRoot* UPropertyBagSchema::GetHierarchyRoot(const TArray<UObject*>& ObjectsWithProperty) const
-{
-	return nullptr;
 }
 
-TSubclassOf<UPropertyBagHierarchyProperty> UPropertyBagSchema::GetHierarchyPropertyType(const FPropertyBagPropertyDesc& Desc) const
-{
-	if (Desc.IsNumericType() && Desc.ValueType != EPropertyBagPropertyType::Bool && Desc.ContainerTypes.IsEmpty())
-	{			
-		return UPropertyBagHierarchyNumericProperty::StaticClass();
-	}
-	
-	return UPropertyBagHierarchyProperty::StaticClass();
-}
-
-TSubclassOf<UPropertyBagHierarchyCategory> UPropertyBagSchema::GetHierarchyCategoryType() const
-{
-	return UPropertyBagHierarchyCategory::StaticClass();
-}
-
-TSubclassOf<UPropertyBagHierarchySection> UPropertyBagSchema::GetHierarchySectionType() const
-{
-	return UPropertyBagHierarchySection::StaticClass();
-}
+////////////////////////////////////
 
-void UPropertyBagSchema::TransferPropertyBagMetadataIntoHierarchy(const FInstancedPropertyBag& PropertyBag,	UPropertyBagHierarchyRoot& HierarchyRoot) const
+bool UPropertyBagSchema::SupportsPinTypeContainer(TWeakPtr<const FEdGraphSchemaAction> SchemaAction,
+		const FEdGraphPinType& PinType, const EPinContainerType& ContainerType) const
 {
-	TMap<const FPropertyBagPropertyDesc*, TArray<FString>> CategoryChains;
-	
-	for (const FPropertyBagPropertyDesc& Desc : PropertyBag.GetPropertyBagStruct()->GetPropertyDescs())
-	{
-		FString CategoryString = Desc.GetMetaData(UE::StructUtils::Metadata::CategoryName);
-		TArray<FString> CategoriesOfDesc;
-		CategoryString.ParseIntoArray(CategoriesOfDesc, TEXT("|"), true);
-		
-		if (CategoriesOfDesc.Num() > 0)
-		{
-			TArray<FString>& Array = CategoryChains.Add(&Desc);
-			
-			for (const FString& Category : CategoriesOfDesc)
-			{
-				Array.Add(Category);
-			}
-		}
-	}
-	
-	auto GenerateHierarchyRecursive = [&](auto&& Self, const FPropertyBagPropertyDesc& Desc, const TArray<FString>& CategoryChain, UHierarchyElement* Outer) -> void
-	{
-		if (CategoryChain.Num() > 0)
-		{				
-			TArray<UPropertyBagHierarchyCategory*> ExistingCategories;
-			Outer->GetChildrenOfType<UPropertyBagHierarchyCategory>(ExistingCategories, false);
-				
-			FString CurrentCategory = CategoryChain[0];
-			
-			UPropertyBagHierarchyCategory* HierarchyCategory = nullptr;
-		
-			UPropertyBagHierarchyCategory** ExistingHierarchyCategory = ExistingCategories.FindByPredicate([CurrentCategory](UPropertyBagHierarchyCategory* Candidate)
-			{
-				return Candidate->GetCategoryName().ToString().Equals(CurrentCategory);
-			});
-		
-			if (ExistingHierarchyCategory)
-			{
-				HierarchyCategory = *ExistingHierarchyCategory;
-			}
-			else
-			{
-				HierarchyCategory = NewObject<UPropertyBagHierarchyCategory>(Outer, GetHierarchyCategoryType());
-				HierarchyCategory->SetCategoryName(FName(CurrentCategory));
-				Outer->GetChildrenMutable().Add(HierarchyCategory);
-			}
-			
-			TArray<FString> ChildCategoryChain = CategoryChain;
-			ChildCategoryChain.RemoveAt(0);
-			Self(Self, Desc, ChildCategoryChain, HierarchyCategory);
-		}
-		else
-		{
-			TSubclassOf<UPropertyBagHierarchyProperty> PropertyClass = GetHierarchyPropertyType(Desc);
-			UPropertyBagHierarchyProperty* HierarchyProperty = NewObject<UPropertyBagHierarchyProperty>(Outer, PropertyClass);
-			HierarchyProperty->Initialize(Desc);
-			Outer->GetChildrenMutable().Add(HierarchyProperty);
-			
-			if (Desc.HasMetaData(UE::StructUtils::Metadata::Specifiers::ToolTipName))
-			{
-				HierarchyProperty->Tooltip = FText::FromString(Desc.GetMetaData(UE::StructUtils::Metadata::Specifiers::ToolTipName));
-			}
-			
-			if (UPropertyBagHierarchyNumericProperty* NumericProperty = Cast<UPropertyBagHierarchyNumericProperty>(HierarchyProperty))
-			{
-				if (Desc.HasMetaData(UE::StructUtils::Metadata::Specifiers::ClampMinName))
-				{
-					float ClampMinValue;
-					if (LexTryParseString(ClampMinValue, *Desc.GetMetaData(UE::StructUtils::Metadata::Specifiers::ClampMinName)))
-					{
-						NumericProperty->bUseClampMin = true;
-						NumericProperty->ClampMin = ClampMinValue;	
-					}
-				}
-				
-				if (Desc.HasMetaData(UE::StructUtils::Metadata::Specifiers::ClampMaxName))
-				{
-					float ClampMaxValue;
-					if (LexTryParseString(ClampMaxValue, *Desc.GetMetaData(UE::StructUtils::Metadata::Specifiers::ClampMaxName)))
-					{
-						NumericProperty->bUseClampMax = true;
-						NumericProperty->ClampMax = ClampMaxValue;	
-					}
-				}
-				
-				if (Desc.HasMetaData(UE::StructUtils::Metadata::Specifiers::UIMinName))
-				{
-					float UIMinValue;
-					if (LexTryParseString(UIMinValue, *Desc.GetMetaData(UE::StructUtils::Metadata::Specifiers::UIMinName)))
-					{
-						NumericProperty->bUseUIMin = true;
-						NumericProperty->UIMin = UIMinValue;	
-					}
-				}
-				
-				if (Desc.HasMetaData(UE::StructUtils::Metadata::Specifiers::UIMaxName))
-				{
-					float UIMaxValue;
-					if (LexTryParseString(UIMaxValue, *Desc.GetMetaData(UE::StructUtils::Metadata::Specifiers::UIMaxName)))
-					{
-						NumericProperty->bUseUIMax = true;
-						NumericProperty->UIMax = UIMaxValue;	
-					}
-				}
-				
-				if (Desc.HasMetaData(UE::StructUtils::Metadata::Specifiers::UnitsName))
-				{
-					FString UnitString = Desc.GetMetaData(UE::StructUtils::Metadata::Specifiers::UnitsName);
-					if (TOptional<EUnit> Unit = FUnitConversion::UnitFromString(*UnitString))
-					{
-						NumericProperty->bUseUnit = true;
-						NumericProperty->Unit = Unit.GetValue();
-					}
-				}
-			}
-		}
-	};
-	
-	for (const auto& CategoryChain : CategoryChains)
-	{
-		GenerateHierarchyRecursive(GenerateHierarchyRecursive, *CategoryChain.Key, CategoryChain.Value, &HierarchyRoot);
-	}
+	return ContainerType == EPinContainerType::None || ContainerType == EPinContainerType::Array || ContainerType == EPinContainerType::Set || ContainerType == EPinContainerType::Map;
 }
 
 #undef LOCTEXT_NAMESPACE
diff --git a/Engine/Source/Editor/StructUtilsEditor/Private/HierarchyEditor/PropertyBagHierarchyViewModel.cpp b/Engine/Source/Editor/StructUtilsEditor/Private/HierarchyEditor/PropertyBagHierarchyViewModel.cpp
deleted file mode 100644
diff --git a/Engine/Source/Editor/StructUtilsEditor/Private/HierarchyEditor/SPropertyBagHierarchyEditor.cpp b/Engine/Source/Editor/StructUtilsEditor/Private/HierarchyEditor/SPropertyBagHierarchyEditor.cpp
deleted file mode 100644
diff --git a/Engine/Source/Editor/StructUtilsEditor/Private/HierarchyEditor/SPropertyBagHierarchyEditor.h b/Engine/Source/Editor/StructUtilsEditor/Private/HierarchyEditor/SPropertyBagHierarchyEditor.h
deleted file mode 100644
diff --git a/Engine/Source/Editor/StructUtilsEditor/Public/HierarchyEditor/PropertyBagHierarchyViewModel.h b/Engine/Source/Editor/StructUtilsEditor/Public/HierarchyEditor/PropertyBagHierarchyViewModel.h
deleted file mode 100644
diff --git a/Engine/Source/Editor/StructUtilsEditor/Public/PropertyBagDetails.h b/Engine/Source/Editor/StructUtilsEditor/Public/PropertyBagDetails.h
--- a/Engine/Source/Editor/StructUtilsEditor/Public/PropertyBagDetails.h
+++ b/Engine/Source/Editor/StructUtilsEditor/Public/PropertyBagDetails.h
@@ -13,7 +13,6 @@ class FReply;
 enum class ESelectorType : uint8;
 struct FEdGraphSchemaAction;
 struct FMenuEntryParams;
-struct EVisibility;
 
 class IPropertyHandle;
 class IPropertyUtilities;
@@ -21,13 +20,6 @@ class IDetailPropertyRow;
 class SInlineEditableTextBlock;
 class SWidget;
 
-class UPropertyBagSchema;
-class UHierarchySection;
-class UPropertyBagHierarchyRoot;
-class UPropertyBagHierarchyProperty;
-class UPropertyBagHierarchySection;
-class UPropertyBagHierarchyCategory;
-
 /**
  * The property bag details child rows can be completely customized by assigning a combination of these values
  * to their FPropertyBag 'ChildRowFeatures' metadata via the meta specifier. TODO: It isn't currently possible
@@ -85,22 +77,16 @@ namespace UE::StructUtils
 
 	/** @return property bag struct common to all edited properties. */
 	STRUCTUTILSEDITOR_API const UPropertyBag* GetCommonBagStruct(TSharedPtr<IPropertyHandle> StructProperty);
-	
-	STRUCTUTILSEDITOR_API TSubclassOf<UPropertyBagSchema> ExtractPropertyBagSchemaClass(TSharedRef<IPropertyHandle> InPropertyBagHandle);
-	STRUCTUTILSEDITOR_API const UPropertyBagSchema* ExtractPropertyBagSchemaCDO(TSharedRef<IPropertyHandle> InPropertyBagHandle);
-	STRUCTUTILSEDITOR_API UPropertyBagHierarchyRoot* ExtractHierarchyRoot(TSharedRef<IPropertyHandle> InPropertyBagHandle);
-	
+
 	/**
 	 * Allows the mutation of the property descriptors on a Property Bag property via the provided callback.
 	 * Handles the whole transaction internally, including building a new Property Bag Struct (if necessary).
 	 *
 	 * @param SessionName Name of the transaction
 	 * @param StructProperty The Property Bag handle
+	 * @param PropUtils Property Utilities, usually taken from the StructCustomizationUtils in CustomizeHeader
 	 * @param Function The callback that does the mutation
 	 */
-	STRUCTUTILSEDITOR_API void ApplyChangesToPropertyDescs(const FText& SessionName, const TSharedPtr<IPropertyHandle>& StructProperty, TFunctionRef<void(TArray<FPropertyBagPropertyDesc>&)> Function);
-
-	UE_DEPRECATED(5.8, "Please use the version without the property utilities parameter.")
 	STRUCTUTILSEDITOR_API void ApplyChangesToPropertyDescs(const FText& SessionName, const TSharedPtr<IPropertyHandle>& StructProperty, const TSharedPtr<IPropertyUtilities>& PropUtils, TFunctionRef<void(TArray<FPropertyBagPropertyDesc>&)> Function);
 
 	/**
@@ -110,9 +96,10 @@ namespace UE::StructUtils
 	 * @param SessionName Name of the transaction
 	 * @param PropertyHandle A child handle of the Property Bag to mutate
 	 * @param StructProperty The Property Bag handle
+	 * @param PropUtils Property Utilities, usually taken from the StructCustomizationUtils in CustomizeHeader
 	 * @param Function The callback that does the mutation
 	 */
-	STRUCTUTILSEDITOR_API void ApplyChangesToSinglePropertyDesc(const FText& SessionName, const TSharedPtr<IPropertyHandle> PropertyHandle, const TSharedPtr<IPropertyHandle>& StructProperty, TFunctionRef<void(FPropertyBagPropertyDesc&)> Function);
+	STRUCTUTILSEDITOR_API void ApplyChangesToSinglePropertyDesc(const FText& SessionName, const TSharedPtr<IPropertyHandle> PropertyHandle, const TSharedPtr<IPropertyHandle>& StructProperty, const TSharedPtr<IPropertyUtilities>& PropUtils, TFunctionRef<void(FPropertyBagPropertyDesc&)> Function);
 
 	/**
 	 * Allows the mutation of a single property descriptor on a Property Bag property via the provided callback.
@@ -121,11 +108,9 @@ namespace UE::StructUtils
 	 * @param SessionName Name of the transaction
 	 * @param PropertyDesc The child descriptor to edit
 	 * @param StructProperty The Property Bag handle
+	 * @param PropUtils Property Utilities, usually taken from the StructCustomizationUtils in CustomizeHeader
 	 * @param Function The callback that does the mutation
 	 */
-	STRUCTUTILSEDITOR_API void ApplyChangesToSinglePropertyDesc(const FText& SessionName, const FPropertyBagPropertyDesc& PropertyDesc, const TSharedPtr<IPropertyHandle>& StructProperty, TFunctionRef<void(FPropertyBagPropertyDesc&)> Function);
-	
-	UE_DEPRECATED(5.8, "Please use the version without the property utilities parameter.")
 	STRUCTUTILSEDITOR_API void ApplyChangesToSinglePropertyDesc(const FText& SessionName, const FPropertyBagPropertyDesc& PropertyDesc, const TSharedPtr<IPropertyHandle>& StructProperty, const TSharedPtr<IPropertyUtilities>& PropUtils, TFunctionRef<void(FPropertyBagPropertyDesc&)> Function);
 }
 
@@ -155,31 +140,20 @@ protected:
 	/** IPropertyTypeCustomization interface */
 	UE_API virtual void CustomizeHeader(TSharedRef<IPropertyHandle> StructPropertyHandle, FDetailWidgetRow& HeaderRow, IPropertyTypeCustomizationUtils& StructCustomizationUtils) override;
 	UE_API virtual void CustomizeChildren(TSharedRef<IPropertyHandle> StructPropertyHandle, IDetailChildrenBuilder& StructBuilder, IPropertyTypeCustomizationUtils& StructCustomizationUtils) override;
-	
-	static FReply SummonHierarchyEditor(TSharedPtr<IPropertyHandle> InPropertyBagHandle, TSharedPtr<IPropertyUtilities> InPropertyUtilities);
-	void NavigateToHierarchyProperty(const FPropertyBagPropertyDesc& Desc);
-	
-	static void ClosePropertyBagHierarchyEditorWindow();
+
 private:
-	/** Handle to the property bag being edited */
+	/** Handle to the struct property being edited */
 	TSharedPtr<IPropertyHandle> StructProperty;
 	TSharedPtr<IPropertyUtilities> PropUtils;
 	EPropertyBagPropertyType DefaultType = EPropertyBagPropertyType::Bool;
 	EPropertyBagChildRowFeatures ChildRowFeatures = EPropertyBagChildRowFeatures::Default;
-	
-	// We keep track of existing window & hierarchy editor loosely so we can recycle them where appropriate
-	// This also means we can only have 1 property bag hierarchy editor open at a time
-	// This is because an additional window can't rely on the existing details panel and property handles too much
-	static TWeakPtr<SWindow> PropertyBagHierarchyEditorWindow;
-	static TWeakPtr<class SPropertyBagHierarchyEditor> PropertyBagHierarchyEditor;
 
 	UE_DEPRECATED(5.6, "Use 'ChildRowFeatures' instead.")
 	bool bFixedLayout = false;
 	bool bAllowContainers = true;
 	
 	// Default to UPropertyBagSchema, can be set with metadata PropertyBagSchemaClass
-	TWeakObjectPtr<const UPropertyBagSchema> PropertyBagSchemaCDO;
-	TWeakObjectPtr<UPropertyBagHierarchyRoot> HierarchyRoot;
+	TSubclassOf<UEdGraphSchema> PropertyBagSchemaClass;
 };
 
 /** 
@@ -193,22 +167,13 @@ private:
 class FPropertyBagInstanceDataDetails : public FInstancedStructDataDetails
 {
 public:
-	/** Generic delegate to execute on a property description */
-	DECLARE_DELEGATE_OneParam(FPropertyBagPropertyDelegate, const FPropertyBagPropertyDesc& PropertyDesc)
-	
 	struct FConstructParams
 	{
-		// Hierarchy Params are only meaningful if the owning property bag has a valid HierarchyRoot provided by a schema
-		struct FHierarchyParams
-		{
-			FPropertyBagPropertyDelegate OnNavigateToHierarchyPropertyRequested;
-		};
-		
 		TSharedPtr<IPropertyHandle> BagStructProperty = nullptr;
 		TSharedPtr<IPropertyUtilities> PropUtils = nullptr;
 		bool bAllowContainers = true;
 		EPropertyBagChildRowFeatures ChildRowFeatures = EPropertyBagChildRowFeatures::Default;
-		FHierarchyParams HierarchyParams;
+		TSubclassOf<UEdGraphSchema> PropertyBagSchemaClass;
 	};
 
 	UE_API explicit FPropertyBagInstanceDataDetails(const FConstructParams& ConstructParams);
@@ -223,8 +188,7 @@ public:
 	FPropertyBagInstanceDataDetails& operator=(FPropertyBagInstanceDataDetails&& Other) noexcept = delete;
 	UE_API PRAGMA_ENABLE_DEPRECATION_WARNINGS
 
-	virtual void AddChildRows(IDetailChildrenBuilder& ChildBuilder, const TArray<TSharedPtr<IPropertyHandle>>& ChildProperties) override;
-	UE_API virtual void OnGroupRowAdded(IDetailGroup& GroupRow, int32 Level, const FString& Category) const override;
+	virtual void OnGroupRowAdded(IDetailGroup& GroupRow, int32 Level, const FString& Category) const override;
 	UE_API virtual void OnChildRowAdded(IDetailPropertyRow& ChildRow) override;
 	
 	/** Enum describing if a property is overridden, or undetermined (e.g. multiselection) */
@@ -298,21 +262,7 @@ public:
 
 	/** Called to reset the child property to default value. */
 	UE_API virtual void ResetToDefault(TSharedPtr<IPropertyHandle> ChildPropertyHandle);
-	
-	virtual void SetOnRebuildChildren(FSimpleDelegate InOnRegenerateChildren) override
-	{
-		OnRebuildChildren = InOnRegenerateChildren;
-	}
-	
-	/**~ Public Hierarchy Editor API */	
-	UE_API EVisibility ShouldShowHierarchyProperty(TSharedPtr<IPropertyHandle> ChildPropertyHandle) const;
-	
-	TWeakObjectPtr<const UHierarchySection> GetActiveHierarchySection() const { return HierarchyData.ActiveHierarchySection; }
-	
-	/** Retrieves the hierarchy property handle map. Only meaningful if used in conjunction with a hierarchy root. */
-	const TMap<const UPropertyBagHierarchyProperty*, TSharedPtr<IPropertyHandle>>& GetHierarchyPropertyHandleMap() const { return HierarchyData.HierarchyPropertyHandleMap; }
-	/**~ End of Hierarchy Editor API */
-	
+
 protected:
 	UE_DEPRECATED(5.6, "Replaced by 'CreatePropertyDetailsWidget' to allow for more customizable drop-down menu location and content.")
 	UE_API TSharedRef<SWidget> OnPropertyNameContent(TSharedPtr<IPropertyHandle> ChildPropertyHandle, TSharedPtr<SInlineEditableTextBlock> InlineWidget) const;
@@ -320,8 +270,6 @@ protected:
 	TSharedPtr<IPropertyHandle> BagStructProperty;
 	TSharedPtr<IPropertyUtilities> PropUtils;
 
-	FSimpleDelegate OnRebuildChildren;
-	
 	bool bAllowContainers = true;
 	EPropertyBagChildRowFeatures ChildRowFeatures = EPropertyBagChildRowFeatures::Default;
 
@@ -330,35 +278,8 @@ protected:
 	UE_DEPRECATED(5.6, "Use 'bAllowContainers' instead.")
 	bool bAllowArrays = true;
 
-	UE_DEPRECATED(5.8, "Use PropertyBagSchema instead. Will be invalid. PropertyBagSchema is set automatically via property metadata instead of passed in via construction arguments now.")
-	TSubclassOf<UEdGraphSchema> PropertyBagSchemaClass;
-	
 	// Default to UPropertyBagSchema, can be set with metadata PropertyBagSchemaClass
-	TWeakObjectPtr<const UPropertyBagSchema> PropertyBagSchema;
-	
-	/**~ Hierarchy Editor API */
-	void GenerateHierarchySectionRow(IDetailChildrenBuilder& ChildrenBuilder);
-	void GenerateHierarchyRootRows(IDetailChildrenBuilder& ChildrenBuilder);
-	void GenerateHierarchyLeftoverPropertyRows(IDetailChildrenBuilder& ChildrenBuilder, const TArray<TSharedPtr<IPropertyHandle>>& AllChildPropertyHandles);
-	void GenerateHierarchyShowAdvancedRow(IDetailChildrenBuilder& ChildrenBuilder);
-	
-	const FSlateBrush* GetAdvancedImage() const;
-	EVisibility GetAdvancedButtonVisibility() const;
-	
-	TSharedPtr<SWidget> ShowAdvancedButton;
-	bool bShouldShowAdvanced = false;
-	
-	struct FHierarchyData
-	{
-		TWeakObjectPtr<UPropertyBagHierarchyRoot> HierarchyRoot;
-		TMap<const UPropertyBagHierarchyProperty*, TSharedPtr<IPropertyHandle>> HierarchyPropertyHandleMap;
-		TWeakObjectPtr<const UHierarchySection> ActiveHierarchySection;
-	
-		FPropertyBagPropertyDelegate OnNavigateToHierarchyPropertyRequestedDelegate;
-	};
-	
-	FHierarchyData HierarchyData;
-	/**~ End of Hierarchy Editor API */	
+	TSubclassOf<UEdGraphSchema> PropertyBagSchemaClass;
 };
 
 UCLASS(MinimalAPI)
@@ -382,21 +303,6 @@ class UPropertyBagSchema : public UEdGraphSchema_K2
 	GENERATED_BODY()
 public:
 	UE_API virtual bool SupportsPinTypeContainer(TWeakPtr<const FEdGraphSchemaAction> SchemaAction, const FEdGraphPinType& PinType, const EPinContainerType& ContainerType) const override;
-	
-	/** Retrieves the hierarchy root from the given objects that own the property bag. Optional. */
-	UE_API virtual UPropertyBagHierarchyRoot* GetHierarchyRoot(const TArray<UObject*>& ObjectsWithProperty) const;
-
-	/** Defines the property type within the hierarchy. Useful if you want to subclass and add your own properties. */
-	UE_API virtual TSubclassOf<UPropertyBagHierarchyProperty> GetHierarchyPropertyType(const FPropertyBagPropertyDesc& Desc) const;
-	
-	/** Defines the category type within the hierarchy. Useful if you want to subclass and add your own properties. */
-	UE_API virtual TSubclassOf<UPropertyBagHierarchyCategory> GetHierarchyCategoryType() const;
-	
-	/** Defines the section type within the hierarchy. Useful if you want to subclass and add your own properties. */
-	UE_API virtual TSubclassOf<UPropertyBagHierarchySection> GetHierarchySectionType() const;
-	
-	/** Post Load helper function. Override if you know what you are doing. */
-	UE_API virtual void TransferPropertyBagMetadataIntoHierarchy(const FInstancedPropertyBag& PropertyBag, UPropertyBagHierarchyRoot& HierarchyRoot) const; 
 };
 
 #undef UE_API
diff --git a/Engine/Source/Editor/StructUtilsEditor/StructUtilsEditor.Build.cs b/Engine/Source/Editor/StructUtilsEditor/StructUtilsEditor.Build.cs
--- a/Engine/Source/Editor/StructUtilsEditor/StructUtilsEditor.Build.cs
+++ b/Engine/Source/Editor/StructUtilsEditor/StructUtilsEditor.Build.cs
@@ -29,7 +29,6 @@ namespace UnrealBuildTool.Rules
 					"AIModule",
 					"DetailCustomizations",
 					"ComponentVisualizers",
-					"DataHierarchyEditor"
 				}
 			);
 
@@ -43,7 +42,7 @@ namespace UnrealBuildTool.Rules
 					"AIGraph",
 					"ToolMenus",
 					"BlueprintGraph",
-					"KismetCompiler"
+					"KismetCompiler",
 				}
 			);
 		}
diff --git a/Engine/Source/Runtime/CoreUObject/Private/StructUtils/PropertyBag.cpp b/Engine/Source/Runtime/CoreUObject/Private/StructUtils/PropertyBag.cpp
--- a/Engine/Source/Runtime/CoreUObject/Private/StructUtils/PropertyBag.cpp
+++ b/Engine/Source/Runtime/CoreUObject/Private/StructUtils/PropertyBag.cpp
@@ -66,6 +66,80 @@ namespace UE::StructUtils::Private
 		return From != nullptr && To != nullptr && From->IsChildOf(To);
 	}
 
+	void AppendPropertyDescs(FXxHash64Builder& HashBuilder, const FPropertyBagPropertyDesc& Desc);
+
+	void AppendPathName(FStringBuilderBase& StringBuilder, const UObject* Object)
+	{
+		if (Object)
+		{
+			Object->GetPathName(nullptr, StringBuilder);
+		}
+	}
+
+	void AppendValueObjectHash(FXxHash64Builder& HashBuilder, FStringBuilderBase& StringBuilder, const UObject* InValueObject)
+	{
+		AppendPathName(StringBuilder, InValueObject);
+#if WITH_EDITOR
+		// for user defined structures we need to hash each property. after changing a user defined structure
+		// we may find the same property bag again if the hash is not sufficient, even though the memory layout
+		// has changed. hashing the path name of the user defined structure is not enough.
+		if(const UUserDefinedStruct* UserDefinedStruct = Cast<UUserDefinedStruct>(InValueObject))
+		{
+			for(const FField* Property = UserDefinedStruct->ChildProperties; Property != nullptr; Property = Property->Next)
+			{
+				const uint32 FieldHash = GetTypeHash(Property);
+				HashBuilder.Update(&FieldHash, sizeof(FieldHash));
+				if(const FStructProperty* StructProperty = CastField<FStructProperty>(Property))
+				{
+					if(StructProperty->Struct)
+					{
+						AppendValueObjectHash(HashBuilder, StringBuilder, StructProperty->Struct);
+					}
+				}
+			}
+		}
+		// for property bags we recurse
+		else if(const UPropertyBag* PropertyBag = Cast<UPropertyBag>(InValueObject))
+		{
+			for (const FPropertyBagPropertyDesc& Desc : PropertyBag->GetPropertyDescs())
+			{
+				AppendPropertyDescs(HashBuilder, Desc);
+			}
+		}
+#endif
+	}
+
+	void AppendPropertyDescs(FXxHash64Builder& HashBuilder, const FPropertyBagPropertyDesc& Desc)
+	{
+		TStringBuilder<256> StringBuilder;
+		HashBuilder.Update(&Desc.ID, sizeof(Desc.ID));
+		StringBuilder << Desc.Name;
+		HashBuilder.Update(&Desc.ValueType, sizeof(Desc.ValueType));
+		HashBuilder.Update(&Desc.ContainerTypes, sizeof(Desc.ContainerTypes));
+		HashBuilder.Update(&Desc.PropertyFlags, sizeof(Desc.PropertyFlags));
+		HashBuilder.Update(&Desc.KeyType, sizeof(Desc.KeyType));
+#if WITH_EDITORONLY_DATA
+		for (const FPropertyBagPropertyDescMetaData& MetaData : Desc.MetaData)
+		{
+			StringBuilder << MetaData.Key;
+			HashBuilder.Update(*MetaData.Value, MetaData.Value.Len() * sizeof(FString::ElementType));
+		}
+#endif
+		AppendValueObjectHash(HashBuilder, StringBuilder, Desc.ValueTypeObject);
+		AppendValueObjectHash(HashBuilder, StringBuilder, Desc.KeyTypeObject);
+		HashBuilder.Update(StringBuilder.ToString(), StringBuilder.Len() * sizeof(TStringBuilder<256>::ElementType));
+	}
+
+	uint64 CalcPropertyDescArrayHash(const TConstArrayView<FPropertyBagPropertyDesc> Descs)
+	{
+		FXxHash64Builder HashBuilder;
+		for (const FPropertyBagPropertyDesc& Desc : Descs)
+		{
+			AppendPropertyDescs(HashBuilder, Desc);
+		}
+		return HashBuilder.Finalize().Hash;
+	}
+
 	bool IsTypeHashable(const EPropertyBagPropertyType InType, const UObject* InTypeObject)
 	{
 		switch (InType)
@@ -4356,7 +4430,7 @@ TArray<FPropertyBagPropertyDesc> UPropertyBag::SanitizePropertyDescs(const TArra
 const UPropertyBag* UPropertyBag::GetOrCreateFromDescs(const TConstArrayView<FPropertyBagPropertyDesc> InPropertyDescs, const TCHAR* PrefixName)
 {
 	TArray<FPropertyBagPropertyDesc> PropertyDescs = SanitizePropertyDescs(InPropertyDescs);
-	const uint64 BagHash = UE::StructUtils::CalcPropertyDescArrayHash(PropertyDescs);
+	const uint64 BagHash = UE::StructUtils::Private::CalcPropertyDescArrayHash(PropertyDescs);
 	const FString ScriptStructName = PrefixName == nullptr
 		? FString::Printf(TEXT("PropertyBag_%llx"), BagHash)
 		: FString::Printf(TEXT("%s_%llx"), PrefixName, BagHash);
diff --git a/Engine/Source/Runtime/CoreUObject/Private/StructUtils/StructUtils.cpp b/Engine/Source/Runtime/CoreUObject/Private/StructUtils/StructUtils.cpp
deleted file mode 100644
diff --git a/Engine/Source/Runtime/CoreUObject/Public/StructUtils/StructUtils.h b/Engine/Source/Runtime/CoreUObject/Public/StructUtils/StructUtils.h
--- a/Engine/Source/Runtime/CoreUObject/Public/StructUtils/StructUtils.h
+++ b/Engine/Source/Runtime/CoreUObject/Public/StructUtils/StructUtils.h
@@ -5,12 +5,6 @@
 #include "Templates/UnrealTypeTraits.h"
 #include "UObject/Class.h"
 
-class UPropertyBagHierarchyRoot;
-struct FInstancedPropertyBag;
-struct FPropertyBagPropertyDesc;
-class FXxHash64Builder;
-struct FPropertyBagPropertyDescMetaData;
-
 namespace UE::StructUtils
 {
 	template <typename T>
@@ -125,8 +119,4 @@ namespace UE::StructUtils
 		// Clamp count between 0 and the distance to the end of the range
 		InOutCount = FMath::Clamp(InOutCount, 0, (ArrayNum - InOutIndex));
 	}
-	
-	COREUOBJECT_API void AppendPropertyDescs(FXxHash64Builder& HashBuilder, const FPropertyBagPropertyDesc& Desc);
-
-	COREUOBJECT_API uint64 CalcPropertyDescArrayHash(const TConstArrayView<FPropertyBagPropertyDesc> Descs);
 }
