PropertyBagDetails: Drop ensure when no common bag struct exists for a multi-object selection.
Such selections can be perfectly normal (e.g. two graph nodes with different bag types).

#rnx

#Synced-CL 50401605

diff --git a/Engine/Source/Editor/StructUtilsEditor/Private/Customizations/PropertyBagDetails.cpp b/Engine/Source/Editor/StructUtilsEditor/Private/Customizations/PropertyBagDetails.cpp
--- a/Engine/Source/Editor/StructUtilsEditor/Private/Customizations/PropertyBagDetails.cpp
+++ b/Engine/Source/Editor/StructUtilsEditor/Private/Customizations/PropertyBagDetails.cpp
@@ -1084,7 +1084,8 @@ void FPropertyBagInstanceDataDetails::OnChildRowAdded(IDetailPropertyRow& ChildR
 	PRAGMA_ENABLE_DEPRECATION_WARNINGS
 
 	// Validate data and check if it's editable
-	if (!ChildProperty || ChildProperty->HasMetaData(Metadata::HideInDetailPanelsName))
+	// If there's no common bag struct, that might just mean that Multiple Values are selected
+	if (!ChildProperty || ChildProperty->HasMetaData(Metadata::HideInDetailPanelsName) || !BagStruct)
 	{
 		ChildRow.Visibility(EVisibility::Collapsed);
 		return;
