Fixed crash where old constructor left a property uninitizalized access nullptr

#Synced-CL 50531866

diff --git a/Engine/Source/Editor/StructUtilsEditor/Private/Customizations/PropertyBagDetails.cpp b/Engine/Source/Editor/StructUtilsEditor/Private/Customizations/PropertyBagDetails.cpp
--- a/Engine/Source/Editor/StructUtilsEditor/Private/Customizations/PropertyBagDetails.cpp
+++ b/Engine/Source/Editor/StructUtilsEditor/Private/Customizations/PropertyBagDetails.cpp
@@ -1439,7 +1439,7 @@ void FPropertyBagInstanceDataDetails::OnChildRowAdded(IDetailPropertyRow& ChildR
 				.Padding(1, 0)
 				.AutoWidth()
 				[
-					CreateTypeSelectionWidget(ChildPropertyHandle, BagStructProperty, PropUtils, SPinTypeSelector::ESelectorType::Compact, bAllowContainers, PropertyBagSchema->GetClass())
+					CreateTypeSelectionWidget(ChildPropertyHandle, BagStructProperty, PropUtils, SPinTypeSelector::ESelectorType::Compact, bAllowContainers, GetPropertyBagSchemaClass())
 				];
 		}
 
@@ -1660,7 +1660,7 @@ void FPropertyBagInstanceDataDetails::OnChildRowAdded(IDetailPropertyRow& ChildR
 						PropUtils,
 						SPinTypeSelector::ESelectorType::Full,
 						bAllowContainers,
-						PropertyBagSchema->GetClass()),
+						GetPropertyBagSchemaClass()),
 					/*InLabel=*/FText::GetEmpty());
 				MenuBuilder.EndSection();
 			}
@@ -2324,7 +2324,7 @@ TSharedRef<SWidget> FPropertyBagInstanceDataDetails::OnPropertyNameContent(TShar
 				PropUtils,
 				SPinTypeSelector::ESelectorType::Full,
 				bAllowContainers,
-				PropertyBagSchema->GetClass())
+				GetPropertyBagSchemaClass())
 		],
 		FText::GetEmpty());
 
@@ -2366,6 +2366,11 @@ TSharedRef<SWidget> FPropertyBagInstanceDataDetails::OnPropertyNameContent(TShar
 	return MenuBuilder.MakeWidget();
 }
 
+TSubclassOf<UPropertyBagSchema> FPropertyBagInstanceDataDetails::GetPropertyBagSchemaClass() const
+{
+	return PropertyBagSchema.IsValid() ? PropertyBagSchema->GetClass() : UPropertyBagSchema::StaticClass();
+}
+
 void FPropertyBagInstanceDataDetails::GenerateHierarchySectionRow(IDetailChildrenBuilder& ChildrenBuilder)
 {
 	if (!HierarchyData.HierarchyRoot.IsValid())
diff --git a/Engine/Source/Editor/StructUtilsEditor/Public/PropertyBagDetails.h b/Engine/Source/Editor/StructUtilsEditor/Public/PropertyBagDetails.h
--- a/Engine/Source/Editor/StructUtilsEditor/Public/PropertyBagDetails.h
+++ b/Engine/Source/Editor/StructUtilsEditor/Public/PropertyBagDetails.h
@@ -337,6 +337,8 @@ protected:
 	// Default to UPropertyBagSchema, can be set with metadata PropertyBagSchemaClass
 	TWeakObjectPtr<const UPropertyBagSchema> PropertyBagSchema;
 	
+	TSubclassOf<UPropertyBagSchema> GetPropertyBagSchemaClass() const;
+	
 	/**~ Hierarchy Editor API */
 	void GenerateHierarchySectionRow(IDetailChildrenBuilder& ChildrenBuilder);
 	void GenerateHierarchyRootRows(IDetailChildrenBuilder& ChildrenBuilder);
