[StructUtils] Add metadata to UserDefinedStruct. Previously, the BlueprintType metadata was enforced when the asset was copied in the ContentBrowser.

[RN] Metadata added programmatically to UserDefinedStruct should now be set on the editor data. They will be added when the UserDefinedStruct compiles.

#Synced-CL 50904874

diff --git a/Engine/Plugins/VirtualProduction/SceneState/Source/SceneStateEventEditor/Private/SceneStateEventEditorUtils.cpp b/Engine/Plugins/VirtualProduction/SceneState/Source/SceneStateEventEditor/Private/SceneStateEventEditorUtils.cpp
--- a/Engine/Plugins/VirtualProduction/SceneState/Source/SceneStateEventEditor/Private/SceneStateEventEditorUtils.cpp
+++ b/Engine/Plugins/VirtualProduction/SceneState/Source/SceneStateEventEditor/Private/SceneStateEventEditorUtils.cpp
@@ -19,13 +19,12 @@ namespace Private
 UUserDefinedStruct* CreateUserDefinedStruct(UObject* InOuter)
 {
 	UUserDefinedStruct* Struct = NewObject<UUserDefinedStruct>(InOuter, NAME_None, RF_Transactional | RF_Public);
-	check(Struct);
 
-	Struct->EditorData = NewObject<UUserDefinedStructEditorData>(Struct, NAME_None, RF_Transactional);
-	check(Struct->EditorData);
+	UUserDefinedStructEditorData* EditorData = NewObject<UUserDefinedStructEditorData>(Struct, NAME_None, RF_Transactional);
+	Struct->EditorData = EditorData;
+	EditorData->MetaData.Add(FBlueprintTags::BlueprintType, TEXT("true"));
 
 	Struct->Guid = FGuid::NewGuid();
-	Struct->SetMetaData(TEXT("BlueprintType"), TEXT("true"));
 	Struct->Bind();
 	Struct->StaticLink(true);
 
diff --git a/Engine/Source/Editor/KismetCompiler/Private/UserDefinedStructureCompilerUtils.cpp b/Engine/Source/Editor/KismetCompiler/Private/UserDefinedStructureCompilerUtils.cpp
--- a/Engine/Source/Editor/KismetCompiler/Private/UserDefinedStructureCompilerUtils.cpp
+++ b/Engine/Source/Editor/KismetCompiler/Private/UserDefinedStructureCompilerUtils.cpp
@@ -387,10 +387,17 @@ struct FUserDefinedStructureCompilerInner
 		check(Struct);
 		const int32 ErrorNum = MessageLog.NumErrors;
 
-		Struct->SetMetaData(FBlueprintMetadata::MD_Tooltip, *FStructureEditorUtils::GetTooltip(Struct));
-
 		UUserDefinedStructEditorData* EditorData = CastChecked<UUserDefinedStructEditorData>(Struct->EditorData);
 
+		UPackage* Package = Struct->GetOutermost();
+		if (ensure(Package))
+		{
+			TMap<FName, FString>& Map = Package->GetMetaData().ObjectMetaDataMap.FindOrAdd(FSoftObjectPath(Struct));
+			Map.Empty();
+			Map.Append(EditorData->MetaData);
+			Map.Add(FBlueprintMetadata::MD_Tooltip, *FStructureEditorUtils::GetTooltip(Struct));
+		}
+
 		CreateVariables(Struct, K2Schema, MessageLog);
 
 		Struct->Bind();
diff --git a/Engine/Source/Editor/StructUtilsEditor/Private/StructUtilsEditorUtils.cpp b/Engine/Source/Editor/StructUtilsEditor/Private/StructUtilsEditorUtils.cpp
--- a/Engine/Source/Editor/StructUtilsEditor/Private/StructUtilsEditorUtils.cpp
+++ b/Engine/Source/Editor/StructUtilsEditor/Private/StructUtilsEditorUtils.cpp
@@ -7,6 +7,7 @@
 #include "UObject/Package.h"
 #include "UObject/UnrealType.h"
 
+#include "Blueprint/BlueprintSupport.h"
 #include "Kismet2/BlueprintEditorUtils.h"
 #include "PropertyBagDetails.h"
 
@@ -26,12 +27,14 @@ namespace UE::StructUtils::Private
 		Struct->Guid = FGuid::NewGuid();
 		if (Args.bIsBlueprintType)
 		{
-			Struct->SetMetaData("BlueprintType", TEXT("true"));
+			UDSEditorData->MetaData.Add(FBlueprintTags::BlueprintType, TEXT("true"));
 		}
 		if (!Args.bAddPropertyIDToPropertyName)
 		{
-			Struct->SetMetaData(NAME_GeneratePropertyWithoutID, TEXT("true"));
+			UDSEditorData->MetaData.Add(NAME_GeneratePropertyWithoutID, TEXT("true"));
 		}
+		TMap<FName, FString>& Map = Struct->GetPackage()->GetMetaData().ObjectMetaDataMap.FindOrAdd(FSoftObjectPath(Struct));
+		Map.Append(UDSEditorData->MetaData);
 
 		Struct->Bind();
 		constexpr bool bRelinkExistingProperties = true;
@@ -129,7 +132,7 @@ namespace UE::StructUtils
 	bool UpdateUserDefinedStructFromDescs(TNotNull<UUserDefinedStruct*> UserDefinedStruct, const FInstancedPropertyBag& PropertyBag)
 	{
 		bool bEqual = true;
-		UUserDefinedStructEditorData* EditorData = CastChecked<UUserDefinedStructEditorData>(UserDefinedStruct->EditorData);
+		UUserDefinedStructEditorData* EditorData = Cast<UUserDefinedStructEditorData>(UserDefinedStruct->EditorData);
 		if (ensure(EditorData))
 		{
 			TArrayView<const FPropertyBagPropertyDesc> PropertyDescs;
@@ -138,7 +141,8 @@ namespace UE::StructUtils
 				PropertyDescs = PropertyBag.GetPropertyBagStruct()->GetPropertyDescs();
 			}
 
-			const bool bNameWithoutID = UserDefinedStruct->GetMetaData(Private::NAME_GeneratePropertyWithoutID) == TEXT("true");
+			const FString* FoundGeneratePropertyWithoutID = EditorData->MetaData.Find(Private::NAME_GeneratePropertyWithoutID);
+			const bool bNameWithoutID = FoundGeneratePropertyWithoutID ? *FoundGeneratePropertyWithoutID == TEXT("true") : false;
 			TArray<FStructVariableDescription> VariablesDescriptions = Private::GetPropertyDescs(PropertyBag, PropertyDescs, !bNameWithoutID);
 			bEqual = EditorData->VariablesDescriptions.Num() == VariablesDescriptions.Num();
 			if (bEqual)
diff --git a/Engine/Source/Editor/UnrealEd/Classes/UserDefinedStructure/UserDefinedStructEditorData.h b/Engine/Source/Editor/UnrealEd/Classes/UserDefinedStructure/UserDefinedStructEditorData.h
--- a/Engine/Source/Editor/UnrealEd/Classes/UserDefinedStructure/UserDefinedStructEditorData.h
+++ b/Engine/Source/Editor/UnrealEd/Classes/UserDefinedStructure/UserDefinedStructEditorData.h
@@ -149,12 +149,16 @@ public:
 	UPROPERTY()
 	FString ToolTip;
 
+	UPROPERTY()
+	TMap<FName, FString> MetaData;
+
 public:
 	// UObject interface.
 	UNREALED_API virtual TSharedPtr<ITransactionObjectAnnotation> FactoryTransactionAnnotation(const ETransactionAnnotationCreationMode InCreationMode) const override;
 	UNREALED_API virtual void Serialize(FArchive& Ar) override;
 	UNREALED_API virtual void PostEditUndo() override;
 	UNREALED_API virtual void PostEditUndo(TSharedPtr<ITransactionObjectAnnotation> TransactionAnnotation) override;
+	UNREALED_API virtual void PostLoad() override;
 	UNREALED_API virtual void PostLoadSubobjects(FObjectInstancingGraph* OuterInstanceGraph) override;
 	UNREALED_API virtual void PreSave(FObjectPreSaveContext ObjectSaveContext) override;
 	// End of UObject interface.
@@ -179,13 +183,16 @@ public:
 	UNREALED_API virtual FString GetFriendlyNameForProperty(const UUserDefinedStruct* Struct, const FProperty* Property) const override;
 
 private:
-
 	// Track the structure change that PostEditUndo undid to pass to FUserDefinedStructureCompilerUtils::CompileStruct
 	FStructureEditorUtils::EStructureEditorChangeInfo CachedStructureChange;
 
 	// Utility function for both PostEditUndo to route through
 	UNREALED_API void ConsolidatedPostEditUndo(FStructureEditorUtils::EStructureEditorChangeInfo ActiveChange);
 
+	//Whether the MetaData was loaded from the package and saved in UUserDefinedStructEditorData.
+	UPROPERTY()
+	bool bMetaDataConverted = false;
+
 };
 
 #undef UE_API
diff --git a/Engine/Source/Editor/UnrealEd/Private/Kismet2/StructureEditorUtils.cpp b/Engine/Source/Editor/UnrealEd/Private/Kismet2/StructureEditorUtils.cpp
--- a/Engine/Source/Editor/UnrealEd/Private/Kismet2/StructureEditorUtils.cpp
+++ b/Engine/Source/Editor/UnrealEd/Private/Kismet2/StructureEditorUtils.cpp
@@ -46,12 +46,11 @@ UUserDefinedStruct* FStructureEditorUtils::CreateUserDefinedStruct(UObject* InPa
 	if (UserDefinedStructEnabled())
 	{
 		Struct = NewObject<UUserDefinedStruct>(InParent, Name, Flags);
-		check(Struct);
-		Struct->EditorData = NewObject<UUserDefinedStructEditorData>(Struct, NAME_None, RF_Transactional);
-		check(Struct->EditorData);
+		UUserDefinedStructEditorData* EditorData = NewObject<UUserDefinedStructEditorData>(Struct, NAME_None, RF_Transactional);
+		Struct->EditorData = EditorData;
+		EditorData->MetaData.Add(FBlueprintTags::BlueprintType, TEXT("true"));
 
 		Struct->Guid = FGuid::NewGuid();
-		Struct->SetMetaData(TEXT("BlueprintType"), TEXT("true"));
 		Struct->Bind();
 		Struct->StaticLink(true);
 		Struct->Status = UDSS_Error;
diff --git a/Engine/Source/Editor/UnrealEd/Private/UserDefinedStructEditorData.cpp b/Engine/Source/Editor/UnrealEd/Private/UserDefinedStructEditorData.cpp
--- a/Engine/Source/Editor/UnrealEd/Private/UserDefinedStructEditorData.cpp
+++ b/Engine/Source/Editor/UnrealEd/Private/UserDefinedStructEditorData.cpp
@@ -206,6 +206,31 @@ void UUserDefinedStructEditorData::PostEditUndo(TSharedPtr<ITransactionObjectAnn
 	ConsolidatedPostEditUndo(TransactedStructureChange);
 }
 
+void UUserDefinedStructEditorData::PostLoad()
+{
+	Super::PostLoad();
+
+	if (!bMetaDataConverted && !HasAnyFlags(RF_ClassDefaultObject))
+	{
+		const UUserDefinedStruct* ScriptStruct = GetOwnerStruct();
+		if (ensure(ScriptStruct != nullptr))
+		{
+			TMap<FName, FString>* Map = GetPackage()->GetMetaData().GetMapForObject(ScriptStruct);
+			if (Map)
+			{
+				for (const TPair<FName, FString>& Data : *Map)
+				{
+					if (Data.Key != FBlueprintMetadata::MD_Tooltip)
+					{
+						MetaData.Add(Data.Key, Data.Value);
+					}
+				}
+			}
+			bMetaDataConverted = true;
+		}
+	}
+}
+
 void UUserDefinedStructEditorData::PostLoadSubobjects(FObjectInstancingGraph* OuterInstanceGraph)
 {
 	Super::PostLoadSubobjects(OuterInstanceGraph);
diff --git a/Engine/Source/Runtime/CoreUObject/Private/StructUtils/UserDefinedStruct.cpp b/Engine/Source/Runtime/CoreUObject/Private/StructUtils/UserDefinedStruct.cpp
--- a/Engine/Source/Runtime/CoreUObject/Private/StructUtils/UserDefinedStruct.cpp
+++ b/Engine/Source/Runtime/CoreUObject/Private/StructUtils/UserDefinedStruct.cpp
@@ -182,7 +182,6 @@ void UUserDefinedStruct::PostDuplicate(bool bDuplicateForPIE)
 	}
 	if (!bDuplicateForPIE && (GetOuter() != GetTransientPackage()))
 	{
-		SetMetaData(TEXT("BlueprintType"), TEXT("true"));
 		FUserDefinedStructEditorUtils::OnStructureChanged(this);
 	}
 }
