Fixed regenerate children delegate not being bound. Also made it protected so instanced property bag can ask for a refresh itself.

#Synced-CL 50932743

diff --git a/Engine/Source/Editor/StructUtilsEditor/Private/Customizations/PropertyBagDetails.cpp b/Engine/Source/Editor/StructUtilsEditor/Private/Customizations/PropertyBagDetails.cpp
--- a/Engine/Source/Editor/StructUtilsEditor/Private/Customizations/PropertyBagDetails.cpp
+++ b/Engine/Source/Editor/StructUtilsEditor/Private/Customizations/PropertyBagDetails.cpp
@@ -2483,7 +2483,7 @@ void FPropertyBagInstanceDataDetails::GenerateHierarchySectionRow(IDetailChildre
 							.OnCheckStateChanged_Lambda([this, Section](ECheckBoxState NewState)
 							{
 								HierarchyData.ActiveHierarchySection = Section;
-								OnRebuildChildren.ExecuteIfBound();
+								OnRegenerateChildren.ExecuteIfBound();
 							})
 							.IsChecked_Lambda([this, Section]()
 							{
@@ -2515,7 +2515,7 @@ void FPropertyBagInstanceDataDetails::GenerateHierarchySectionRow(IDetailChildre
 					.OnCheckStateChanged_Lambda([this](ECheckBoxState NewState)
 					{
 						HierarchyData.ActiveHierarchySection = nullptr;
-						OnRebuildChildren.ExecuteIfBound();
+						OnRegenerateChildren.ExecuteIfBound();
 					})
 					.IsChecked_Lambda([this]()
 					{
diff --git a/Engine/Source/Editor/StructUtilsEditor/Public/PropertyBagDetails.h b/Engine/Source/Editor/StructUtilsEditor/Public/PropertyBagDetails.h
--- a/Engine/Source/Editor/StructUtilsEditor/Public/PropertyBagDetails.h
+++ b/Engine/Source/Editor/StructUtilsEditor/Public/PropertyBagDetails.h
@@ -301,11 +301,6 @@ public:
 	/** Called to reset the child property to default value. */
 	UE_API virtual void ResetToDefault(TSharedPtr<IPropertyHandle> ChildPropertyHandle);
 	
-	virtual void SetOnRebuildChildren(FSimpleDelegate InOnRegenerateChildren) override
-	{
-		OnRebuildChildren = InOnRegenerateChildren;
-	}
-	
 	/**~ Public Hierarchy Editor API */	
 	UE_API EVisibility ShouldShowHierarchyProperty(TSharedPtr<IPropertyHandle> ChildPropertyHandle) const;
 	
@@ -321,8 +316,6 @@ protected:
 
 	TSharedPtr<IPropertyHandle> BagStructProperty;
 	TSharedPtr<IPropertyUtilities> PropUtils;
-
-	FSimpleDelegate OnRebuildChildren;
 	
 	bool bAllowContainers = true;
 	EPropertyBagChildRowFeatures ChildRowFeatures = EPropertyBagChildRowFeatures::Default;
