[StateTree] Modify UpdateUserDefinedStructFromDescs to test if the UserDefinedStruct contains types that are not supported by Blueprint. The UI prevents adding BlueprintType. But a user can save the asset, remove the BlueprintType from the struct, then reopen the StateTree asset. The UserDefinedStruct compiler doesn?t check for those check but the Blueprint compiler does.

#Synced-CL 50933965

diff --git a/Engine/Plugins/Runtime/StateTree/Source/StateTreeEditorModule/Private/StateTreeCompiler.cpp b/Engine/Plugins/Runtime/StateTree/Source/StateTreeEditorModule/Private/StateTreeCompiler.cpp
--- a/Engine/Plugins/Runtime/StateTree/Source/StateTreeEditorModule/Private/StateTreeCompiler.cpp
+++ b/Engine/Plugins/Runtime/StateTree/Source/StateTreeEditorModule/Private/StateTreeCompiler.cpp
@@ -722,7 +722,11 @@ bool FStateTreeCompiler::CreateParameters()
 			if (StateTree->GlobalParametersStructType)
 			{
 				StateTree->GlobalParameters = StateTree->GlobalParametersStructType;
-				UE::StructUtils::UpdateUserDefinedStructFromDescs(StateTree->GlobalParametersStructType, PropertyBag);
+				if (UE::StructUtils::UpdateUserDefinedStructFromDescs(StateTree->GlobalParametersStructType, PropertyBag).HasError())
+				{
+					// A property type doesn't exist, or the BlueprintType was removed.
+					Log.Reportf(EMessageSeverity::Warning, TEXT("Failed to create a global parameters struct compatible with blueprint type."));
+				}
 			}
 			else
 			{
diff --git a/Engine/Source/Editor/StructUtilsEditor/Private/StructUtilsEditorUtils.cpp b/Engine/Source/Editor/StructUtilsEditor/Private/StructUtilsEditorUtils.cpp
--- a/Engine/Source/Editor/StructUtilsEditor/Private/StructUtilsEditorUtils.cpp
+++ b/Engine/Source/Editor/StructUtilsEditor/Private/StructUtilsEditorUtils.cpp
@@ -58,10 +58,48 @@ namespace UE::StructUtils::Private
 			&& A.DefaultValue == B.DefaultValue;
 	}
 
-	FStructVariableDescription CreateUserDefinedVariableFromDesc(const FPropertyBagPropertyDesc& PropertyDesc, bool bNameWithID)
+	TValueOrError<FStructVariableDescription, void> CreateUserDefinedVariableFromDesc(const FPropertyBagPropertyDesc& PropertyDesc, bool bNameWithID)
 	{
+		if (!UE::StructUtils::CanCreateValidGraphPinTypeForPropertyDesc(PropertyDesc))
+		{
+			return MakeError();
+		}
+
 		const FEdGraphPinType GraphPin = UE::StructUtils::GetPropertyDescAsPin(PropertyDesc);
 
+		auto IsAllowableBlueprintVariableType = [](const UObject* Object)
+			{
+				if (Object)
+				{
+					constexpr bool bForInternalUse = false;
+					if (const UClass* Class = Cast<const UClass>(Object))
+					{
+						return UEdGraphSchema_K2::IsAllowableBlueprintVariableType(Class, bForInternalUse);
+					}
+					else if (const UScriptStruct* ScriptStruct = Cast<const UScriptStruct>(Object))
+					{
+						return UEdGraphSchema_K2::IsAllowableBlueprintVariableType(ScriptStruct, bForInternalUse);
+					}
+					else if (const UEnum* Enum = Cast<const UEnum>(Object))
+					{
+						return UEdGraphSchema_K2::IsAllowableBlueprintVariableType(Enum);
+					}
+				}
+				return true;
+			};
+
+		// Test the struct/object/enum type.
+		if (!IsAllowableBlueprintVariableType(GraphPin.PinSubCategoryObject.Get()))
+		{
+			return MakeError();
+		}
+
+		// Test the value (struct/object/enum) type.for map.
+		if (!IsAllowableBlueprintVariableType(GraphPin.PinValueType.TerminalSubCategoryObject.Get()))
+		{
+			return MakeError();
+		}
+
 		FStructVariableDescription Result;
 		Result.SetPinType(GraphPin);
 		Result.FriendlyName = PropertyDesc.Name.ToString();
@@ -73,31 +111,42 @@ namespace UE::StructUtils::Private
 		}
 		Result.VarGuid = PropertyDesc.ID;
 
-		return Result;
+		return MakeValue(MoveTemp(Result));
 	}
 
-	FStructVariableDescription CreateUserDefinedVariableFromDesc(const FInstancedPropertyBag& PropertyBag, const FPropertyBagPropertyDesc& PropertyDesc, bool bNameWithID)
+	TValueOrError<FStructVariableDescription, void> CreateUserDefinedVariableFromDesc(const FInstancedPropertyBag& InstancedPropertyBag, const FPropertyBagPropertyDesc& PropertyDesc, bool bNameWithID)
 	{
-		FStructVariableDescription Result = CreateUserDefinedVariableFromDesc(PropertyDesc, bNameWithID);
-		if (const FProperty* Property = PropertyBag.GetValue().GetScriptStruct()->FindPropertyByName(PropertyDesc.Name))
+		TValueOrError<FStructVariableDescription, void> Result = CreateUserDefinedVariableFromDesc(PropertyDesc, bNameWithID);
+		if (Result.HasError())
+		{
+			return MakeError();
+		}
+
+		const UScriptStruct* PropertyBag = InstancedPropertyBag.GetValue().GetScriptStruct();
+		if (const FProperty* Property = PropertyBag->FindPropertyByName(PropertyDesc.Name))
 		{
-			UObject* OwningObject = const_cast<UScriptStruct*>(PropertyBag.GetValue().GetScriptStruct());
-			const bool bValueAsString = FBlueprintEditorUtils::PropertyValueToString(Property, PropertyBag.GetValue().GetMemory(), Result.DefaultValue, OwningObject, PPF_DuplicateVerbatim);
+			UObject* OwningObject = const_cast<UScriptStruct*>(PropertyBag);
+			const bool bValueAsString = FBlueprintEditorUtils::PropertyValueToString(Property, InstancedPropertyBag.GetValue().GetMemory(), Result.GetValue().DefaultValue, OwningObject, PPF_DuplicateVerbatim);
 			ensure(bValueAsString);
 		}
 		return Result;
 	}
 
-	TArray<FStructVariableDescription> GetPropertyDescs(const FInstancedPropertyBag& PropertyBag, const TArrayView<const FPropertyBagPropertyDesc>  PropertyDescs, bool bNameWithID)
+	TValueOrError<TArray<FStructVariableDescription>, void> GetPropertyDescs(const FInstancedPropertyBag& PropertyBag, const TArrayView<const FPropertyBagPropertyDesc>  PropertyDescs, bool bNameWithID)
 	{
 		TArray<FStructVariableDescription> VariablesDescriptions;
 		VariablesDescriptions.Reserve(PropertyDescs.Num());
 		TArray<FPropertyBagPropertyDesc> SanitizedPropertyDescs = UPropertyBag::SanitizePropertyDescs(PropertyDescs);
 		for (const FPropertyBagPropertyDesc& Desc : SanitizedPropertyDescs)
 		{
-			VariablesDescriptions.Add(CreateUserDefinedVariableFromDesc(PropertyBag, Desc, bNameWithID));
+			TValueOrError<FStructVariableDescription, void> Variable = CreateUserDefinedVariableFromDesc(PropertyBag, Desc, bNameWithID);
+			if (Variable.HasError())
+			{
+				return MakeError();
+			}
+			VariablesDescriptions.Add(Variable.StealValue());
 		}
-		return VariablesDescriptions;
+		return MakeValue(MoveTemp(VariablesDescriptions));
 	}
 }
 
@@ -122,14 +171,20 @@ namespace UE::StructUtils
 			return MakeValue(UserDefinedStruct);
 		}
 
+		TValueOrError<TArray<FStructVariableDescription>, void> VariableDescriptionsResult = Private::GetPropertyDescs(PropertyBag, PropertyDescs, Args.bAddPropertyIDToPropertyName);
+		if (VariableDescriptionsResult.HasError())
+		{
+			return MakeError();
+		}
+
 		TNotNull<UUserDefinedStructEditorData*> EditorData = CastChecked<UUserDefinedStructEditorData>(UserDefinedStruct->EditorData);
-		EditorData->VariablesDescriptions = Private::GetPropertyDescs(PropertyBag, PropertyDescs, Args.bAddPropertyIDToPropertyName);
+		EditorData->VariablesDescriptions = VariableDescriptionsResult.StealValue();
 		FStructureEditorUtils::CompileStructure(UserDefinedStruct);
 
 		return MakeValue(UserDefinedStruct);
 	}
 
-	bool UpdateUserDefinedStructFromDescs(TNotNull<UUserDefinedStruct*> UserDefinedStruct, const FInstancedPropertyBag& PropertyBag)
+	TValueOrError<bool, void> UpdateUserDefinedStructFromDescs(TNotNull<UUserDefinedStruct*> UserDefinedStruct, const FInstancedPropertyBag& PropertyBag)
 	{
 		bool bEqual = true;
 		UUserDefinedStructEditorData* EditorData = Cast<UUserDefinedStructEditorData>(UserDefinedStruct->EditorData);
@@ -143,7 +198,13 @@ namespace UE::StructUtils
 
 			const FString* FoundGeneratePropertyWithoutID = EditorData->MetaData.Find(Private::NAME_GeneratePropertyWithoutID);
 			const bool bNameWithoutID = FoundGeneratePropertyWithoutID ? *FoundGeneratePropertyWithoutID == TEXT("true") : false;
-			TArray<FStructVariableDescription> VariablesDescriptions = Private::GetPropertyDescs(PropertyBag, PropertyDescs, !bNameWithoutID);
+			TValueOrError<TArray<FStructVariableDescription>, void> VariableDescriptionsResult = Private::GetPropertyDescs(PropertyBag, PropertyDescs, !bNameWithoutID);
+			if (VariableDescriptionsResult.HasError())
+			{
+				return MakeError();
+			}
+
+			TArray<FStructVariableDescription>& VariablesDescriptions = VariableDescriptionsResult.GetValue();
 			bEqual = EditorData->VariablesDescriptions.Num() == VariablesDescriptions.Num();
 			if (bEqual)
 			{
@@ -175,7 +236,7 @@ namespace UE::StructUtils
 			}
 		}
 
-		return !bEqual;
+		return MakeValue(!bEqual);
 	}
 }
 
diff --git a/Engine/Source/Editor/StructUtilsEditor/Private/Tests/StructReinstancerTests.cpp b/Engine/Source/Editor/StructUtilsEditor/Private/Tests/StructReinstancerTests.cpp
--- a/Engine/Source/Editor/StructUtilsEditor/Private/Tests/StructReinstancerTests.cpp
+++ b/Engine/Source/Editor/StructUtilsEditor/Private/Tests/StructReinstancerTests.cpp
@@ -173,8 +173,8 @@ namespace UE::StructUtils::Private
 			BagA.SetValueInt32(BName, 22);
 			BagA.SetValueInt32(CName, 3);
 
-			const bool bModified = UE::StructUtils::UpdateUserDefinedStructFromDescs(UserDefinedStruct, BagA);
-			UTEST_TRUE("UUserDefinedStruct modified", bModified);
+			const TValueOrError<bool, void> ModifyResult = UE::StructUtils::UpdateUserDefinedStructFromDescs(UserDefinedStruct, BagA);
+			UTEST_TRUE("UUserDefinedStruct modified", ModifyResult.HasValue() && ModifyResult.GetValue());
 
 			{
 				UTEST_TRUE(TEXT("TestInstancedObject1.InstancedStruct struct"), TestInstancedObject1->InstancedStruct.GetScriptStruct() == UserDefinedStruct);
@@ -323,8 +323,8 @@ namespace UE::StructUtils::Private
 
 		{
 			BagA.SetValueInt32(BName, -2);
-			const bool bModified = UE::StructUtils::UpdateUserDefinedStructFromDescs(UserDefinedStruct, BagA);
-			UTEST_TRUE("UUserDefinedStruct modified", bModified);
+			const TValueOrError<bool, void> ModifyResult = UE::StructUtils::UpdateUserDefinedStructFromDescs(UserDefinedStruct, BagA);
+			UTEST_TRUE("UUserDefinedStruct modified", ModifyResult.HasValue() && ModifyResult.GetValue());
 		}
 
 		UTEST_TRUE(TEXT("Initial Test Struct"), TestStruct());
@@ -341,8 +341,8 @@ namespace UE::StructUtils::Private
 		// Set the default value of C. Should stay has it is. The child already modified it.
 		{
 			BagA.SetValueInt32(CName, -3);
-			const bool bModified = UE::StructUtils::UpdateUserDefinedStructFromDescs(UserDefinedStruct, BagA);
-			UTEST_TRUE("UUserDefinedStruct modified", bModified);
+			const TValueOrError<bool, void> ModifyResult = UE::StructUtils::UpdateUserDefinedStructFromDescs(UserDefinedStruct, BagA);
+			UTEST_TRUE("UUserDefinedStruct modified", ModifyResult.HasValue() && ModifyResult.GetValue());
 		}
 
 		UTEST_TRUE(TEXT("Initial Test Struct"), TestStruct());
@@ -462,8 +462,8 @@ namespace UE::StructUtils::Private
 			UTEST_TRUE(TEXT("Assign"), EPropertyBagResult::Success == BagA.SetValueInt32(CName, -3));
 			BagA.AddProperties({ { DName, EPropertyBagPropertyType::Int32 } });
 			UTEST_TRUE(TEXT("Assign"), EPropertyBagResult::Success == BagA.SetValueInt32(DName, 4));
-			const bool bModified = UE::StructUtils::UpdateUserDefinedStructFromDescs(UserDefinedStruct, BagA);
-			UTEST_TRUE(TEXT("UUserDefinedStruct modified"), bModified);
+			const TValueOrError<bool, void> ModifyResult = UE::StructUtils::UpdateUserDefinedStructFromDescs(UserDefinedStruct, BagA);
+			UTEST_TRUE(TEXT("UUserDefinedStruct modified"), ModifyResult.HasValue() && ModifyResult.GetValue());
 		}
 		{
 			const UScriptStruct* OldBagBStruct = BagB.GetValue().GetScriptStruct();
diff --git a/Engine/Source/Editor/StructUtilsEditor/Public/StructUtilsEditorUtils.h b/Engine/Source/Editor/StructUtilsEditor/Public/StructUtilsEditorUtils.h
--- a/Engine/Source/Editor/StructUtilsEditor/Public/StructUtilsEditorUtils.h
+++ b/Engine/Source/Editor/StructUtilsEditor/Public/StructUtilsEditorUtils.h
@@ -29,15 +29,18 @@ namespace UE::StructUtils
 
 	/**
 	 * @return a new user defined struct created from a property bag.
-	 * Fails if the name is already used by another asset or if the asset could not be created.
+	 * Fails if the name is already used by another asset
+	 *	or if the asset could not be created
+	 *	or the FInstancedPropertyBag contains type that are not supported by UserDefinedStruct.
 	 */
 	UE_API TValueOrError<UUserDefinedStruct*, void> CreateUserDefinedStructFromDescs(TNotNull<UObject*> UserDefinedStructOuter, const FInstancedPropertyBag& PropertyBag, FName UserDefinedStructName, const FCreateUserDefinedStructArgs& Args);
 
 	/**
 	 * Update, if needed, an existing user defined struct to have the exact same properties and default value as the property bag.
 	 * @return true if an update occurs.
+	 * Fails if FInstancedPropertyBag contains type that are not supported by UserDefinedStruct.
 	 */
-	UE_API bool UpdateUserDefinedStructFromDescs(TNotNull<UUserDefinedStruct*> Struct, const FInstancedPropertyBag& PropertyBag);
+	UE_API TValueOrError<bool, void> UpdateUserDefinedStructFromDescs(TNotNull<UUserDefinedStruct*> Struct, const FInstancedPropertyBag& PropertyBag);
 }
 
 #undef UE_API
