Fixed advanced button showing up when in a section that doesn't contain any advanced parameters

#Synced-CL 51048610

diff --git a/Engine/Source/Editor/StructUtilsEditor/Private/Customizations/PropertyBagDetails.cpp b/Engine/Source/Editor/StructUtilsEditor/Private/Customizations/PropertyBagDetails.cpp
--- a/Engine/Source/Editor/StructUtilsEditor/Private/Customizations/PropertyBagDetails.cpp
+++ b/Engine/Source/Editor/StructUtilsEditor/Private/Customizations/PropertyBagDetails.cpp
@@ -2655,16 +2655,63 @@ EVisibility FPropertyBagInstanceDataDetails::GetAdvancedButtonVisibility() const
 {
 	if (HierarchyData.HierarchyRoot.IsValid())
 	{
-		TArray<UPropertyBagHierarchyProperty*> AllHierarchyProperties;
-		HierarchyData.HierarchyRoot->GetChildrenOfType<UPropertyBagHierarchyProperty>(AllHierarchyProperties, true);
-		
-		for (UPropertyBagHierarchyProperty* Property : AllHierarchyProperties)
+		// If no section (meaning the All section in UI) is active, we check if any property is advanced
+		if (HierarchyData.ActiveHierarchySection == nullptr)
+		{
+			TArray<UPropertyBagHierarchyProperty*> AllHierarchyProperties;
+			HierarchyData.HierarchyRoot->GetChildrenOfType<UPropertyBagHierarchyProperty>(AllHierarchyProperties, true);
+			
+			for(UPropertyBagHierarchyProperty* Prop : AllHierarchyProperties)
+			{
+				if(FPropertyBagPropertyMetadata_Common* CommonMetaData = Prop->FindPropertyMetaDataOfTypeMutable<FPropertyBagPropertyMetadata_Common>())
+				{
+					if(CommonMetaData->bAdvanced)
+					{
+						return EVisibility::Visible;
+					}
+				}
+			}
+		}
+		// If there is an actual section active, we filter for categories within that section and then check their respective properties
+		else
 		{
-			if (FPropertyBagPropertyMetadata_Common* CommonMetaData = Property->FindPropertyMetaDataOfTypeMutable<FPropertyBagPropertyMetadata_Common>())
+			// We first collect immediate categories visible for the current section
+			TArray<const UPropertyBagHierarchyCategory*> Categories;
+			bool bRecursiveCategoryLoopup = false;
+			HierarchyData.HierarchyRoot->GetChildrenOfType(Categories, bRecursiveCategoryLoopup);
+		
+			for(const UPropertyBagHierarchyCategory* Category : Categories)
 			{
-				if (CommonMetaData->bAdvanced)
+				const UHierarchySection* CategorySection = nullptr;
+				
+				// We check if they are visible under the current section
+				if(const FDataHierarchyElementMetaData_SectionAssociation* SectionMetaData = Category->FindMetaDataOfType<FDataHierarchyElementMetaData_SectionAssociation>())
 				{
-					return EVisibility::Visible;
+					if(SectionMetaData->Section.IsValid())
+					{
+						CategorySection = SectionMetaData->Section.Get();
+					}
+				}
+
+				if(CategorySection != HierarchyData.ActiveHierarchySection)
+				{
+					continue;
+				}
+
+				// Check all properties within this category (recursively)
+				TArray<UPropertyBagHierarchyProperty*> Props;
+				bool bRecursivePropertyLookup = true;
+				Category->GetChildrenOfType<UPropertyBagHierarchyProperty>(Props, bRecursivePropertyLookup);
+
+				for(UPropertyBagHierarchyProperty* Prop : Props)
+				{
+					if(FPropertyBagPropertyMetadata_Common* CommonMetaData = Prop->FindPropertyMetaDataOfTypeMutable<FPropertyBagPropertyMetadata_Common>())
+					{
+						if(CommonMetaData->bAdvanced)
+						{
+							return EVisibility::Visible;
+						}
+					}
 				}
 			}
 		}
