[StructUtils] fixed heap-use-after-free issue when trying to retrieve type names during module shutdown
#rnx
#Synced-CL 51096110

diff --git a/Engine/Source/Editor/StructUtilsEditor/Private/StructUtilsEditorModule.cpp b/Engine/Source/Editor/StructUtilsEditor/Private/StructUtilsEditorModule.cpp
--- a/Engine/Source/Editor/StructUtilsEditor/Private/StructUtilsEditorModule.cpp
+++ b/Engine/Source/Editor/StructUtilsEditor/Private/StructUtilsEditorModule.cpp
@@ -19,25 +19,39 @@
 
 IMPLEMENT_MODULE(FStructUtilsEditorModule, StructUtilsEditor)
 
+namespace UE::StructUtilsEditor::Private
+{
+FName FInstancedStructTypeName;
+FName FInstancedPropertyBagTypeName;
+FName FInstancedStructArrayTypeName;
+}
+
 void FStructUtilsEditorModule::StartupModule()
 {
+	using namespace UE::StructUtilsEditor::Private;
+	FInstancedStructTypeName = FInstancedStruct::StaticStruct()->GetFName();
+	FInstancedPropertyBagTypeName = FInstancedPropertyBag::StaticStruct()->GetFName();
+	FInstancedStructArrayTypeName = FInstancedStructArray::StaticStruct()->GetFName();
+
 	// Register the details customizer
 	FPropertyEditorModule& PropertyModule = FModuleManager::LoadModuleChecked<FPropertyEditorModule>("PropertyEditor");
-	PropertyModule.RegisterCustomPropertyTypeLayout(FInstancedStruct::StaticStruct()->GetFName(), FOnGetPropertyTypeCustomizationInstance::CreateStatic(&FInstancedStructDetails::MakeInstance));
-	PropertyModule.RegisterCustomPropertyTypeLayout(FInstancedPropertyBag::StaticStruct()->GetFName(), FOnGetPropertyTypeCustomizationInstance::CreateStatic(&FPropertyBagDetails::MakeInstance));
-	PropertyModule.RegisterCustomPropertyTypeLayout(FInstancedStructArray::StaticStruct()->GetFName(), FOnGetPropertyTypeCustomizationInstance::CreateStatic(&FInstancedStructArrayDetails::MakeInstance));
+	PropertyModule.RegisterCustomPropertyTypeLayout(FInstancedStructTypeName, FOnGetPropertyTypeCustomizationInstance::CreateStatic(&FInstancedStructDetails::MakeInstance));
+	PropertyModule.RegisterCustomPropertyTypeLayout(FInstancedPropertyBagTypeName, FOnGetPropertyTypeCustomizationInstance::CreateStatic(&FPropertyBagDetails::MakeInstance));
+	PropertyModule.RegisterCustomPropertyTypeLayout(FInstancedStructArrayTypeName, FOnGetPropertyTypeCustomizationInstance::CreateStatic(&FInstancedStructArrayDetails::MakeInstance));
 	PropertyModule.NotifyCustomizationModuleChanged();
 }
 
 void FStructUtilsEditorModule::ShutdownModule()
 {
+	using namespace UE::StructUtilsEditor::Private;
+
 	// Unregister the details customization
 	if (FModuleManager::Get().IsModuleLoaded("PropertyEditor"))
 	{
 		FPropertyEditorModule& PropertyModule = FModuleManager::LoadModuleChecked<FPropertyEditorModule>("PropertyEditor");
-		PropertyModule.UnregisterCustomPropertyTypeLayout(FInstancedStruct::StaticStruct()->GetFName());
-		PropertyModule.UnregisterCustomPropertyTypeLayout(FInstancedPropertyBag::StaticStruct()->GetFName());
-		PropertyModule.UnregisterCustomPropertyTypeLayout(FInstancedStructArray::StaticStruct()->GetFName());
+		PropertyModule.UnregisterCustomPropertyTypeLayout(FInstancedStructTypeName);
+		PropertyModule.UnregisterCustomPropertyTypeLayout(FInstancedPropertyBagTypeName);
+		PropertyModule.UnregisterCustomPropertyTypeLayout(FInstancedStructArrayTypeName);
 		PropertyModule.NotifyCustomizationModuleChanged();
 	}
 }
