Changed 'Advanced' reading to read from advanced metadata property directly instead of property flag. This helps with desync issues.
Real fix comes from next CL that allows automatic reparenting when a category with a parameter is deleted, which is what causes the desync in the first place by deleting and recreating the property in the hierarchy with a clean slate (advanced = false despite the CPF_Advanced flag being true). With the real fix properties won't get deleted anymore by deleting the parent category, so data stays in sync

Ideally we get rid of the "double source of truth" with reading from FProperty flags vs. reading from hierarchy property metadata altogether and read only from hierarchy property metadata in the future.

#Synced-CL 51048610

diff --git a/Engine/Source/Editor/StructUtilsEditor/Private/Customizations/PropertyBagDetails.cpp b/Engine/Source/Editor/StructUtilsEditor/Private/Customizations/PropertyBagDetails.cpp
--- a/Engine/Source/Editor/StructUtilsEditor/Private/Customizations/PropertyBagDetails.cpp
+++ b/Engine/Source/Editor/StructUtilsEditor/Private/Customizations/PropertyBagDetails.cpp
@@ -974,9 +974,10 @@ void FHierarchyPropertyCategoryBuilder::GenerateChildContent(IDetailChildrenBuil
 	
 	for (const UPropertyBagHierarchyProperty* ContainedHierarchyProperty : ContainedProperties)
 	{
-		if (OwningDetails.Pin()->GetHierarchyPropertyHandleMap().Contains(ContainedHierarchyProperty))
+		TObjectKey<const UPropertyBagHierarchyProperty> Key(ContainedHierarchyProperty);
+		if (const TSharedPtr<IPropertyHandle>* PropertyHandle = OwningDetails.Pin()->GetHierarchyPropertyHandleMap().Find(Key))
 		{
-			IDetailPropertyRow& PropertyRow = ChildrenBuilder.AddProperty(OwningDetails.Pin()->GetHierarchyPropertyHandleMap()[ContainedHierarchyProperty].ToSharedRef());
+			IDetailPropertyRow& PropertyRow = ChildrenBuilder.AddProperty(PropertyHandle->ToSharedRef());
 			PropertyRow.Visibility(MakeAttributeSP(OwningDetails.Pin().Get(), &FPropertyBagInstanceDataDetails::ShouldShowHierarchyProperty, PropertyRow.GetPropertyHandle()));
 			OwningDetails.Pin()->OnChildRowAdded(PropertyRow);
 		}
@@ -1081,6 +1082,7 @@ void FPropertyBagInstanceDataDetails::AddChildRows(IDetailChildrenBuilder& Child
 		HierarchyData.HierarchyRoot->GetChildrenOfType(AllHierarchyProperties, true);
 		
 		HierarchyData.HierarchyPropertyHandleMap.Empty();
+		HierarchyData.PropertyPathPropertyMap.Empty();
 		for(const UPropertyBagHierarchyProperty* HierarchyProperty : AllHierarchyProperties)
 		{
 			if (const FPropertyBagPropertyDesc* PropertyDesc = Bag->FindPropertyDescByID(HierarchyProperty->GetPropertyId()))
@@ -1092,7 +1094,8 @@ void FPropertyBagInstanceDataDetails::AddChildRows(IDetailChildrenBuilder& Child
 				
 				if (MatchingPropertyHandle)
 				{
-					HierarchyData.HierarchyPropertyHandleMap.Add(HierarchyProperty, *MatchingPropertyHandle);			
+					HierarchyData.HierarchyPropertyHandleMap.Add(TObjectKey<const UPropertyBagHierarchyProperty>(HierarchyProperty), *MatchingPropertyHandle);
+					HierarchyData.PropertyPathPropertyMap.Add(FString((*MatchingPropertyHandle)->GetPropertyPath()), HierarchyProperty);
 				}
 			}
 		}
@@ -2323,9 +2326,18 @@ EVisibility FPropertyBagInstanceDataDetails::ShouldShowHierarchyProperty(TShared
 		return EVisibility::Visible;
 	}
 	
-	if (FProperty* MetaDataProperty = ChildPropertyHandle->GetMetaDataProperty())
+	// We read directly from the common metadata to check if a property is advanced since the property advanced flag might go out of sync.
+	// This can happen when a property is deleted somehow (parent category deleted), which will re-add at the root automatically but with new metadata
+	// The advanced property flag is pointless anyway in customized UI.
+	if (const TWeakObjectPtr<const UPropertyBagHierarchyProperty>* Property = HierarchyData.PropertyPathPropertyMap.Find(FString(ChildPropertyHandle->GetPropertyPath())))
 	{
-		return MetaDataProperty->HasAnyPropertyFlags(CPF_AdvancedDisplay) ? EVisibility::Collapsed : EVisibility::Visible;
+		if (Property != nullptr && Property->IsValid())
+		{
+			if (const FPropertyBagPropertyMetadata_Common* CommonMetaData = (*Property)->FindPropertyMetaDataOfType<FPropertyBagPropertyMetadata_Common>())
+			{
+				return CommonMetaData->bAdvanced ? EVisibility::Collapsed : EVisibility::Visible;
+			}
+		}
 	}
 	
 	return EVisibility::Visible;
diff --git a/Engine/Source/Editor/StructUtilsEditor/Public/PropertyBagDetails.h b/Engine/Source/Editor/StructUtilsEditor/Public/PropertyBagDetails.h
--- a/Engine/Source/Editor/StructUtilsEditor/Public/PropertyBagDetails.h
+++ b/Engine/Source/Editor/StructUtilsEditor/Public/PropertyBagDetails.h
@@ -307,7 +307,7 @@ public:
 	TWeakObjectPtr<const UHierarchySection> GetActiveHierarchySection() const { return HierarchyData.ActiveHierarchySection; }
 	
 	/** Retrieves the hierarchy property handle map. Only meaningful if used in conjunction with a hierarchy root. */
-	const TMap<const UPropertyBagHierarchyProperty*, TSharedPtr<IPropertyHandle>>& GetHierarchyPropertyHandleMap() const { return HierarchyData.HierarchyPropertyHandleMap; }
+	const TMap<TObjectKey<const UPropertyBagHierarchyProperty>, TSharedPtr<IPropertyHandle>>& GetHierarchyPropertyHandleMap() const { return HierarchyData.HierarchyPropertyHandleMap; }
 	/**~ End of Hierarchy Editor API */
 	
 protected:
@@ -348,7 +348,8 @@ protected:
 	struct FHierarchyData
 	{
 		TWeakObjectPtr<UPropertyBagHierarchyRoot> HierarchyRoot;
-		TMap<const UPropertyBagHierarchyProperty*, TSharedPtr<IPropertyHandle>> HierarchyPropertyHandleMap;
+		TMap<TObjectKey<const UPropertyBagHierarchyProperty>, TSharedPtr<IPropertyHandle>> HierarchyPropertyHandleMap;
+		TMap<FString, TWeakObjectPtr<const UPropertyBagHierarchyProperty>> PropertyPathPropertyMap;
 		TWeakObjectPtr<const UHierarchySection> ActiveHierarchySection;
 	
 		FPropertyBagPropertyDelegate OnNavigateToHierarchyPropertyRequestedDelegate;
