[StructUtils] Fix GetFilteredVariableTypeTree to use the correct schema. Fix the backward compatible FPropertyBagInstanceDataDetails constructor to set the schema.

#Synced-CL 51266215

diff --git a/Engine/Source/Editor/BlueprintGraph/Classes/EdGraphSchema_K2.h b/Engine/Source/Editor/BlueprintGraph/Classes/EdGraphSchema_K2.h
--- a/Engine/Source/Editor/BlueprintGraph/Classes/EdGraphSchema_K2.h
+++ b/Engine/Source/Editor/BlueprintGraph/Classes/EdGraphSchema_K2.h
@@ -480,7 +480,7 @@ public:
 
 	//////////////////////////////////////////////////////////////////////////
 	// FPinTypeInfo
-	/** Class used for creating type tree selection info, which aggregates the various PC_* and PinSubtypes in the schema into a heirarchy */
+	/** Class used for creating type tree selection info, which aggregates the various PC_* and PinSubtypes in the schema into a hierarchy */
 	class FPinTypeTreeInfo
 	{
 	private:
@@ -1080,7 +1080,7 @@ public:
 	 * @param	TypeTree		The array that will contain the type tree hierarchy for this schema upon returning
 	 * @param	TypeTreeFilter	ETypeTreeFilter flags that determine how the TypeTree is populated.
 	 */
-	UE_API void GetVariableTypeTree(TArray< TSharedPtr<FPinTypeTreeInfo> >& TypeTree, ETypeTreeFilter TypeTreeFilter = ETypeTreeFilter::None) const;
+	UE_API virtual void GetVariableTypeTree(TArray< TSharedPtr<FPinTypeTreeInfo> >& TypeTree, ETypeTreeFilter TypeTreeFilter = ETypeTreeFilter::None) const;
 
 	/**
 	 * Returns whether or not the specified type has valid subtypes available
diff --git a/Engine/Source/Editor/StructUtilsEditor/Private/Customizations/PropertyBagDetails.cpp b/Engine/Source/Editor/StructUtilsEditor/Private/Customizations/PropertyBagDetails.cpp
--- a/Engine/Source/Editor/StructUtilsEditor/Private/Customizations/PropertyBagDetails.cpp
+++ b/Engine/Source/Editor/StructUtilsEditor/Private/Customizations/PropertyBagDetails.cpp
@@ -544,7 +544,7 @@ void CopyPropertyValue(const FPropertyBagPropertyDesc* InSourcePropertyDesc, con
 	InSourcePropertyDesc->CachedProperty->CopyCompleteValue(TargetValueAddress, SourceValueAddress);
 }
 
-void GetFilteredVariableTypeTree(const TSharedPtr<IPropertyHandle>& BagStructProperty, TArray<TSharedPtr<UEdGraphSchema_K2::FPinTypeTreeInfo>>& TypeTree, const ETypeTreeFilter TypeTreeFilter)
+void GetFilteredVariableTypeTree(const TSharedPtr<IPropertyHandle>& BagStructProperty, TArray<TSharedPtr<UEdGraphSchema_K2::FPinTypeTreeInfo>>& TypeTree, const ETypeTreeFilter TypeTreeFilter, const UPropertyBagSchema* Schema)
 {
 	// The type selector popup might outlive this details view, so bag struct property can be invalid here.
 	if (!BagStructProperty || !BagStructProperty->IsValidHandle())
@@ -573,9 +573,8 @@ void GetFilteredVariableTypeTree(const TSharedPtr<IPropertyHandle>& BagStructPro
 		}
 	};
 
-	check(GetDefault<UEdGraphSchema_K2>());
 	TArray<TSharedPtr<UEdGraphSchema_K2::FPinTypeTreeInfo>> TempTypeTree;
-	GetDefault<UPropertyBagSchema>()->GetVariableTypeTree(TempTypeTree, TypeTreeFilter);
+	Schema->GetVariableTypeTree(TempTypeTree, TypeTreeFilter);
 
 	// Filter
 	for (TSharedPtr<UEdGraphSchema_K2::FPinTypeTreeInfo>& PinType : TempTypeTree)
@@ -729,9 +728,13 @@ TSharedRef<SWidget> CreateTypeSelectionWidget(TSharedPtr<IPropertyHandle> ChildP
 		.Padding(FMargin(4, 0))
 		[
 			SNew(STypeSelector, FGetPinTypeTree::CreateLambda(
-				[BagStructProperty = InBagStructProperty](TArray<TSharedPtr<UEdGraphSchema_K2::FPinTypeTreeInfo>>& TypeTree, const ETypeTreeFilter TypeTreeFilter)
+				[BagStructProperty = InBagStructProperty, SchemaWeak = TWeakObjectPtr<const UPropertyBagSchema>(Schema)]
+				(TArray<TSharedPtr<UEdGraphSchema_K2::FPinTypeTreeInfo>>& TypeTree, const ETypeTreeFilter TypeTreeFilter)
 				{
-				UE::StructUtils::Private::GetFilteredVariableTypeTree(BagStructProperty, TypeTree, TypeTreeFilter);
+					if (const UPropertyBagSchema* Schema = SchemaWeak.Get())
+					{
+						UE::StructUtils::Private::GetFilteredVariableTypeTree(BagStructProperty, TypeTree, TypeTreeFilter, Schema);
+					}
 				}))
 			.TargetPinType_Lambda([ChildPropertyHandle = ChildPropertyHandle, BagStructProperty = InBagStructProperty]()
 			{
@@ -1004,13 +1007,12 @@ PRAGMA_ENABLE_DEPRECATION_WARNINGS
 {
 	ensure(UE::StructUtils::Private::IsScriptStruct<FInstancedPropertyBag>(BagStructProperty));
 	ensure(PropUtils != nullptr);
-	
+
 	if (BagStructProperty.IsValid() && BagStructProperty->IsValidHandle())
 	{
 		PropertyBagSchema = UE::StructUtils::ExtractPropertyBagSchemaCDO(BagStructProperty.ToSharedRef());
-		
+
 		HierarchyData.HierarchyRoot = UE::StructUtils::ExtractHierarchyRoot(BagStructProperty.ToSharedRef());
-		
 		if (HierarchyData.HierarchyRoot.IsValid())
 		{
 			HierarchyData.OnNavigateToHierarchyPropertyRequestedDelegate = ConstructParams.HierarchyParams.OnNavigateToHierarchyPropertyRequested;
@@ -1032,6 +1034,12 @@ PRAGMA_ENABLE_DEPRECATION_WARNINGS
 {
 	ensure(UE::StructUtils::Private::IsScriptStruct<FInstancedPropertyBag>(BagStructProperty));
 	ensure(PropUtils != nullptr);
+
+	if (BagStructProperty.IsValid() && BagStructProperty->IsValidHandle())
+	{
+		PropertyBagSchema = UE::StructUtils::ExtractPropertyBagSchemaCDO(BagStructProperty.ToSharedRef());
+		HierarchyData.HierarchyRoot = UE::StructUtils::ExtractHierarchyRoot(BagStructProperty.ToSharedRef());
+	}
 }
 
 void FPropertyBagInstanceDataDetails::GenerateHierarchyShowAdvancedRow(IDetailChildrenBuilder& ChildrenBuilder)
