static_assert text doesn't need to be wrapped with TEXT(). This causes errors on newer C++ versions.

#Synced-CL 51569266

diff --git a/Engine/Source/Editor/StructUtilsEditor/Public/HierarchyEditor/PropertyBagHierarchyViewModel.h b/Engine/Source/Editor/StructUtilsEditor/Public/HierarchyEditor/PropertyBagHierarchyViewModel.h
--- a/Engine/Source/Editor/StructUtilsEditor/Public/HierarchyEditor/PropertyBagHierarchyViewModel.h
+++ b/Engine/Source/Editor/StructUtilsEditor/Public/HierarchyEditor/PropertyBagHierarchyViewModel.h
@@ -127,7 +127,7 @@ public:
 	template<typename T>
 	bool ContainsPropertyMetaDataOfType()
 	{
-		static_assert(TIsDerivedFrom<T, FPropertyBagPropertyMetadata_Base>::IsDerived, TEXT("MetaData must be of type FPropertyBagPropertyMetadata_Base!"));
+		static_assert(TIsDerivedFrom<T, FPropertyBagPropertyMetadata_Base>::IsDerived, "MetaData must be of type FPropertyBagPropertyMetadata_Base!");
 		return ContainsPropertyMetaDataOfType(T::StaticStruct);
 	}
 	
@@ -138,7 +138,7 @@ public:
 	template<typename T>
 	const T* FindPropertyMetaDataOfType() const
 	{
-		static_assert(TIsDerivedFrom<T, FPropertyBagPropertyMetadata_Base>::IsDerived, TEXT("MetaData must be of type FPropertyBagPropertyMetadata_Base!"));
+		static_assert(TIsDerivedFrom<T, FPropertyBagPropertyMetadata_Base>::IsDerived, "MetaData must be of type FPropertyBagPropertyMetadata_Base!");
 		
 		if (const FInstancedStruct* Result = FindPropertyMetaDataOfType(T::StaticStruct()))
 		{
@@ -151,7 +151,7 @@ public:
 	template<typename T>
 	T* FindPropertyMetaDataOfTypeMutable()
 	{
-		static_assert(TIsDerivedFrom<T, FPropertyBagPropertyMetadata_Base>::IsDerived, TEXT("MetaData must be of type FPropertyBagPropertyMetadata_Base!"));
+		static_assert(TIsDerivedFrom<T, FPropertyBagPropertyMetadata_Base>::IsDerived, "MetaData must be of type FPropertyBagPropertyMetadata_Base!");
 		
 		if (FInstancedStruct* Result = FindPropertyMetaDataOfTypeMutable(T::StaticStruct()))
 		{
