Unified drag & drop in property bag details via hierarchy view model.
Old path is still intact.

Property Drag & Drop Handler has an alternate Hierarchy path now.
New Category Drag & Drop Handler for hierarchies specifically has been introduced.

Removed some self-explanatory comments that cluttered the code

#Synced-CL 51539883

diff --git a/Engine/Source/Editor/DataHierarchyEditor/Public/DataHierarchyViewModelBase.h b/Engine/Source/Editor/DataHierarchyEditor/Public/DataHierarchyViewModelBase.h
--- a/Engine/Source/Editor/DataHierarchyEditor/Public/DataHierarchyViewModelBase.h
+++ b/Engine/Source/Editor/DataHierarchyEditor/Public/DataHierarchyViewModelBase.h
@@ -890,10 +890,10 @@ struct FHierarchyElementViewModel : TSharedFromThis<FHierarchyElementViewModel>,
 	UE_API bool IsExpandedByDefault() const;
 	
 	/** Returns true if the given item can be dropped on the given target area. */
-	FResultWithUserFeedback CanDropOn(TSharedPtr<FHierarchyElementViewModel> DraggedElement, EItemDropZone ItemDropZone);
+	UE_API FResultWithUserFeedback CanDropOn(TSharedPtr<FHierarchyElementViewModel> DraggedElement, EItemDropZone ItemDropZone);
 
 	/** Gets executed when an item was dropped on this. */
-	void OnDroppedOn(TSharedPtr<FHierarchyElementViewModel> DroppedItem, EItemDropZone ItemDropZone);
+	UE_API void OnDroppedOn(TSharedPtr<FHierarchyElementViewModel> DroppedItem, EItemDropZone ItemDropZone);
 	/** Lets you completely customize the OnDroppedOn behavior. Needs to a specific value if you override this. */
 	virtual TOptional<FResultWithUserFeedback> CanDropOnOverride(TSharedPtr<FHierarchyElementViewModel> DraggedElement, EItemDropZone ItemDropZone) { return {}; }
 	virtual TOptional<bool> OnDroppedOnOverride(TSharedPtr<FHierarchyElementViewModel> DroppedElement, EItemDropZone ItemDropZone) { return {}; }
diff --git a/Engine/Source/Editor/StructUtilsEditor/Private/Customizations/PropertyBagDetails.cpp b/Engine/Source/Editor/StructUtilsEditor/Private/Customizations/PropertyBagDetails.cpp
--- a/Engine/Source/Editor/StructUtilsEditor/Private/Customizations/PropertyBagDetails.cpp
+++ b/Engine/Source/Editor/StructUtilsEditor/Private/Customizations/PropertyBagDetails.cpp
@@ -13,8 +13,10 @@
 #include "IPropertyUtilities.h"
 #include "PropertyBagDragDropHandler.h"
 #include "PropertyCustomizationHelpers.h"
+#include "PropertyPath.h"
 #include "ScopedTransaction.h"
 #include "SDraggableBox.h" // Custom widget for drag and drop sections
+#include "StructUtilsEditorModule.h"
 #include "StructUtilsMetadata.h"
 #include "STypeSelector.h" // Custom widget for pill type selector
 #include "Framework/Application/SlateApplication.h"
@@ -908,6 +910,81 @@ namespace Constants
 } // UE::StructUtils
 
 
+/** Drag & drop handler for hierarchy categories in the details panel. Supports both dragging (source) and dropping (target).
+ *  Uses FHierarchyDragDropOp for dragging (categories are not properties), and accepts both FHierarchyDragDropOp and
+ *  FPropertyBagDetailsDragDropOp for drops (properties dragged onto categories use the latter). */
+class FHierarchyCategoryDetailsDragDropHandler : public IDetailDragDropHandler
+{
+public:
+	/** The delegate receives the full drag drop event so it can update decorators, plus the extracted element and drop zone. */
+	DECLARE_DELEGATE_RetVal_ThreeParams(TOptional<EItemDropZone>, FCanAcceptHierarchyDrop, const FDragDropEvent& /*DragDropEvent*/, TSharedPtr<FHierarchyElementViewModel> /*DraggedElement*/, EItemDropZone /*DropZone*/);
+	DECLARE_DELEGATE_RetVal_TwoParams(FReply, FOnHierarchyDrop, TSharedPtr<FHierarchyElementViewModel> /*DroppedElement*/, EItemDropZone /*DropZone*/);
+
+	FHierarchyCategoryDetailsDragDropHandler(TSharedPtr<FHierarchyElementViewModel> InCategoryVM)
+		: CategoryVM(InCategoryVM)
+	{
+	}
+
+	virtual TSharedPtr<FDragDropOperation> CreateDragDropOperation() const override
+	{
+		TSharedPtr<FHierarchyElementViewModel> PinnedVM = CategoryVM.Pin();
+		if (!PinnedVM.IsValid())
+		{
+			return nullptr;
+		}
+
+		TSharedPtr<FHierarchyDragDropOp> DragOp = MakeShared<FHierarchyDragDropOp>(PinnedVM);
+		DragOp->SetLabel(PinnedVM->ToStringAsText());
+		DragOp->Construct();
+		return DragOp;
+	}
+
+	virtual bool UseHandleWidget() const override { return true; }
+
+	void BindCanAcceptDrop(FCanAcceptHierarchyDrop&& InDelegate) { CanAcceptDropDelegate = MoveTemp(InDelegate); }
+	void BindOnDrop(FOnHierarchyDrop&& InDelegate) { OnDropDelegate = MoveTemp(InDelegate); }
+
+protected:
+	virtual TOptional<EItemDropZone> CanAcceptDrop(const FDragDropEvent& DragDropSource, EItemDropZone DropZone) const override
+	{
+		TSharedPtr<FHierarchyElementViewModel> DraggedElement = ExtractDraggedElement(DragDropSource);
+		if (DraggedElement.IsValid() && CanAcceptDropDelegate.IsBound())
+		{
+			return CanAcceptDropDelegate.Execute(DragDropSource, DraggedElement, DropZone);
+		}
+		return TOptional<EItemDropZone>();
+	}
+
+	virtual bool AcceptDrop(const FDragDropEvent& DragDropSource, EItemDropZone DropZone) const override
+	{
+		TSharedPtr<FHierarchyElementViewModel> DraggedElement = ExtractDraggedElement(DragDropSource);
+		if (DraggedElement.IsValid() && OnDropDelegate.IsBound())
+		{
+			return OnDropDelegate.Execute(DraggedElement, DropZone).IsEventHandled();
+		}
+		return false;
+	}
+
+private:
+	/** Extracts the hierarchy element from either FPropertyBagDetailsDragDropOp or FHierarchyDragDropOp. */
+	static TSharedPtr<FHierarchyElementViewModel> ExtractDraggedElement(const FDragDropEvent& DragDropEvent)
+	{
+		if (const TSharedPtr<FPropertyBagDetailsDragDropOp> PropertyOp = DragDropEvent.GetOperationAs<FPropertyBagDetailsDragDropOp>())
+		{
+			return PropertyOp->GetHierarchyElement();
+		}
+		if (const TSharedPtr<FHierarchyDragDropOp> HierarchyOp = DragDropEvent.GetOperationAs<FHierarchyDragDropOp>())
+		{
+			return HierarchyOp->GetDraggedElement().Pin();
+		}
+		return nullptr;
+	}
+
+	TWeakPtr<FHierarchyElementViewModel> CategoryVM;
+	FCanAcceptHierarchyDrop CanAcceptDropDelegate;
+	FOnHierarchyDrop OnDropDelegate;
+};
+
 /** The category builder will display all sub-categories & parameters contained within a given category in the details panel. */
 class FHierarchyPropertyCategoryBuilder : public IDetailCustomNodeBuilder, public TSharedFromThis<FHierarchyPropertyCategoryBuilder>
 {
@@ -945,6 +1022,67 @@ void FHierarchyPropertyCategoryBuilder::GenerateHeaderRowContent(FDetailWidgetRo
 	];
 
 	NodeRow.FilterString(FText::FromName(CustomBuilderRowName));
+
+	// Add drag & drop support for categories via the hierarchy view model
+	if (OwningDetails.IsValid() && HierarchyCategory.IsValid())
+	{
+		UPropertyBagHierarchyViewModel* HierarchyVM = OwningDetails.Pin()->GetHierarchyViewModel();
+		if (HierarchyVM)
+		{
+			TSharedPtr<FHierarchyRootViewModel> RootVM = HierarchyVM->GetHierarchyRootViewModel();
+			TSharedPtr<FHierarchyElementViewModel> CategoryVM = RootVM ? RootVM->FindViewModelForChild(HierarchyCategory.Get(), true) : nullptr;
+
+			if (CategoryVM.IsValid())
+			{
+				TSharedPtr<FHierarchyCategoryDetailsDragDropHandler> DragDropHandler = MakeShared<FHierarchyCategoryDetailsDragDropHandler>(CategoryVM);
+
+				DragDropHandler->BindCanAcceptDrop(
+					FHierarchyCategoryDetailsDragDropHandler::FCanAcceptHierarchyDrop::CreateLambda(
+						[CategoryVM](const FDragDropEvent& DragDropEvent, TSharedPtr<FHierarchyElementViewModel> DraggedElementVM, EItemDropZone DropZone) -> TOptional<EItemDropZone>
+						{
+							if (!DraggedElementVM.IsValid())
+							{
+								return TOptional<EItemDropZone>();
+							}
+
+							FHierarchyElementViewModel::FResultWithUserFeedback Result = CategoryVM->CanDropOn(DraggedElementVM, DropZone);
+
+							// Update decorators on whichever op type is active
+							if (const TSharedPtr<FPropertyBagDetailsDragDropOp> PropertyOp = DragDropEvent.GetOperationAs<FPropertyBagDetailsDragDropOp>())
+							{
+								TOptional<FPropertyBagDetailsDragDropOp::FDecoration> DecorationOverride;
+								if (Result.UserFeedback.IsSet())
+								{
+									const FSlateBrush* Brush = FAppStyle::Get().GetBrush(Result.bResult ? "Graph.ConnectorFeedback.OK" : "Graph.ConnectorFeedback.Error");
+									DecorationOverride.Emplace(Result.UserFeedback.GetValue(), Brush);
+								}
+								PropertyOp->SetDecoration(Result.bResult ? EPropertyBagDropState::Valid : EPropertyBagDropState::Invalid, MoveTemp(DecorationOverride));
+							}
+							else if (const TSharedPtr<FHierarchyDragDropOp> HierarchyOp = DragDropEvent.GetOperationAs<FHierarchyDragDropOp>())
+							{
+								HierarchyOp->SetDescription(Result.UserFeedback.IsSet() ? Result.UserFeedback.GetValue() : FText::GetEmpty());
+							}
+
+							return Result.bResult ? DropZone : TOptional<EItemDropZone>();
+						}));
+
+				DragDropHandler->BindOnDrop(
+					FHierarchyCategoryDetailsDragDropHandler::FOnHierarchyDrop::CreateLambda(
+						[CategoryVM](TSharedPtr<FHierarchyElementViewModel> DroppedElementVM, EItemDropZone DropZone) -> FReply
+						{
+							if (!DroppedElementVM.IsValid())
+							{
+								return FReply::Unhandled();
+							}
+
+							CategoryVM->OnDroppedOn(DroppedElementVM, DropZone);
+							return FReply::Handled();
+						}));
+
+				NodeRow.DragDropHandler(MoveTemp(DragDropHandler));
+			}
+		}
+	}
 }
 
 bool FHierarchyPropertyCategoryBuilder::InitiallyCollapsed() const
@@ -1016,6 +1154,7 @@ PRAGMA_ENABLE_DEPRECATION_WARNINGS
 		if (HierarchyData.HierarchyRoot.IsValid())
 		{
 			HierarchyData.OnNavigateToHierarchyPropertyRequestedDelegate = ConstructParams.HierarchyParams.OnNavigateToHierarchyPropertyRequested;
+			HierarchyData.HierarchyViewModel = ConstructParams.HierarchyParams.HierarchyViewModel;
 		}
 	}
 }
@@ -2019,16 +2158,86 @@ void FPropertyBagInstanceDataDetails::OnChildRowAdded(IDetailPropertyRow& ChildR
 				];
 		}
 
-		/** DRAG AND DROP HANDLER 
-		 *  We only add drag & drop for reordering if there is no hierarchy root. */
+		/** DRAG AND DROP HANDLER */
 		if (EnumHasAnyFlags(ChildRowFeatures, EPropertyBagChildRowFeatures::DragAndDrop) && PropertyDesc != nullptr)
 		{
 			TSharedPtr<FPropertyBagDetailsDragDropHandler> DragDropHandler = MakeShared<FPropertyBagDetailsDragDropHandler>(*PropertyDesc);
 
-			// If we support drag & drop but a hierarchy root exists, we allow dragging, but not dropping on our property rows
-			if (!HierarchyData.HierarchyRoot.IsValid())
+			if (HierarchyData.HierarchyRoot.IsValid() && HierarchyData.HierarchyViewModel.IsValid())
+			{
+				// Hierarchy-based drag & drop: delegate to the hierarchy view model API
+				UPropertyBagHierarchyViewModel* HierarchyVM = HierarchyData.HierarchyViewModel.Get();
+				TSharedPtr<FHierarchyRootViewModel> RootVM = HierarchyVM->GetHierarchyRootViewModel();
+
+				// Look up the FHierarchyElementViewModel for this property
+				FHierarchyElementIdentity Identity = UPropertyBagHierarchyProperty::ConstructIdentity(*PropertyDesc);
+				TSharedPtr<FHierarchyElementViewModel> ElementVM = RootVM ? RootVM->FindViewModelForChild(Identity, true) : nullptr;
+
+				if (ElementVM.IsValid())
+				{
+					// Set the hierarchy element on the handler so it's attached to drag ops
+					DragDropHandler->SetHierarchyElement(ElementVM);
+
+					// Can accept drag and drop — delegate to CanDropOn
+					DragDropHandler->BindCanAcceptDragDrop(
+						FCanAcceptPropertyBagDetailsRowDropOp::CreateLambda(
+							[ElementVM](const TSharedPtr<FPropertyBagDetailsDragDropOp>& DropOp, EItemDropZone DropZone) -> TOptional<EItemDropZone>
+							{
+								if (!DropOp.IsValid())
+								{
+									return TOptional<EItemDropZone>();
+								}
+
+								TSharedPtr<FHierarchyElementViewModel> DraggedElementVM = DropOp->GetHierarchyElement();
+								if (!DraggedElementVM.IsValid())
+								{
+									DropOp->SetDecoration(EPropertyBagDropState::Invalid);
+									return TOptional<EItemDropZone>();
+								}
+
+								FHierarchyElementViewModel::FResultWithUserFeedback Result = ElementVM->CanDropOn(DraggedElementVM, DropZone);
+								if (Result.bResult)
+								{
+									TOptional<FPropertyBagDetailsDragDropOp::FDecoration> DecorationOverride;
+									if (Result.UserFeedback.IsSet())
+									{
+										const FSlateBrush* Brush = FAppStyle::Get().GetBrush("Graph.ConnectorFeedback.OK");
+										DecorationOverride.Emplace(Result.UserFeedback.GetValue(), Brush);
+									}
+									DropOp->SetDecoration(EPropertyBagDropState::Valid, MoveTemp(DecorationOverride));
+									return DropZone;
+								}
+								else
+								{
+									TOptional<FPropertyBagDetailsDragDropOp::FDecoration> DecorationOverride;
+									if (Result.UserFeedback.IsSet())
+									{
+										const FSlateBrush* Brush = FAppStyle::Get().GetBrush("Graph.ConnectorFeedback.Error");
+										DecorationOverride.Emplace(Result.UserFeedback.GetValue(), Brush);
+									}
+									DropOp->SetDecoration(EPropertyBagDropState::Invalid, MoveTemp(DecorationOverride));
+									return TOptional<EItemDropZone>();
+								}
+							}));
+
+					// Handle the drop — delegate to OnDroppedOn. Uses the hierarchy element directly from the op.
+					DragDropHandler->BindOnHandleHierarchyElementDrop(
+						FOnPropertyBagHierarchyElementDropOp::CreateLambda(
+							[ElementVM](TSharedPtr<FHierarchyElementViewModel> DroppedElementVM, EItemDropZone DropZone) -> FReply
+							{
+								if (!DroppedElementVM.IsValid())
+								{
+									return FReply::Unhandled();
+								}
+
+								ElementVM->OnDroppedOn(DroppedElementVM, DropZone);
+								return FReply::Handled();
+							}));
+				}
+			}
+			else if (!HierarchyData.HierarchyRoot.IsValid())
 			{
-				// Can accept drag and drop check for if this is a valid drop
+				// Non-hierarchy path: reorder via FInstancedPropertyBag::ReorderProperty
 				DragDropHandler->BindCanAcceptDragDrop(
 					FCanAcceptPropertyBagDetailsRowDropOp::CreateLambda(
 						[PropertyDesc](const TSharedPtr<FPropertyBagDetailsDragDropOp>& DropOp, EItemDropZone DropZone) -> TOptional<EItemDropZone>
@@ -2351,6 +2560,21 @@ EVisibility FPropertyBagInstanceDataDetails::ShouldShowHierarchyProperty(TShared
 	return EVisibility::Visible;
 }
 
+TWeakObjectPtr<const UHierarchySection> FPropertyBagInstanceDataDetails::GetActiveHierarchySection() const
+{
+	return HierarchyData.ActiveHierarchySection;
+}
+
+const TMap<TObjectKey<const UPropertyBagHierarchyProperty>, TSharedPtr<IPropertyHandle>>& FPropertyBagInstanceDataDetails::GetHierarchyPropertyHandleMap() const
+{
+	return HierarchyData.HierarchyPropertyHandleMap;
+}
+
+UPropertyBagHierarchyViewModel* FPropertyBagInstanceDataDetails::GetHierarchyViewModel() const
+{
+	return HierarchyData.HierarchyViewModel.Get();
+}
+
 TSharedRef<SWidget> FPropertyBagInstanceDataDetails::OnPropertyNameContent(TSharedPtr<IPropertyHandle> ChildPropertyHandle, TSharedPtr<SInlineEditableTextBlock> InlineWidget) const
 {
 	constexpr bool bInShouldCloseWindowAfterMenuSelection = true;
@@ -2832,8 +3056,33 @@ void FPropertyBagDetails::CustomizeHeader(TSharedRef<IPropertyHandle> StructProp
 
 void FPropertyBagDetails::CustomizeChildren(TSharedRef<IPropertyHandle> StructPropertyHandle, IDetailChildrenBuilder& StructBuilder, IPropertyTypeCustomizationUtils& StructCustomizationUtils)
 {
+	// Extract the hierarchy root for this property bag (if a schema provides one)
+	if (StructProperty.IsValid() && StructProperty->IsValidHandle())
+	{
+		HierarchyRoot = UE::StructUtils::ExtractHierarchyRoot(StructProperty.ToSharedRef());
+	}
+
+	// Create or reuse the hierarchy view model if a hierarchy root exists
+	if (HierarchyRoot.IsValid() && StructProperty.IsValid() && StructProperty->IsValidHandle() && PropUtils.IsValid())
+	{
+		TArray<UObject*> OuterObjects;
+		StructProperty->GetOuterObjects(OuterObjects);
+		TSharedPtr<FPropertyPath> PropertyPath = StructProperty->CreateFPropertyPath();
+
+		if (OuterObjects.Num() == 1 && PropertyPath.IsValid())
+		{
+			FStructUtilsEditorModule& Module = FModuleManager::GetModuleChecked<FStructUtilsEditorModule>("StructUtilsEditor");
+			HierarchyViewModel.Reset(Module.FindOrCreateHierarchyViewModel(HierarchyRoot.Get(), OuterObjects[0], PropertyPath.ToSharedRef(), PropUtils.ToSharedRef()));
+		}
+	}
+	else
+	{
+		HierarchyViewModel.Reset();
+	}
+
 	FPropertyBagInstanceDataDetails::FConstructParams::FHierarchyParams HierarchyParams;
-	HierarchyParams.OnNavigateToHierarchyPropertyRequested = FPropertyBagInstanceDataDetails::FPropertyBagPropertyDelegate::CreateSP(this, &FPropertyBagDetails::NavigateToHierarchyProperty); 
+	HierarchyParams.OnNavigateToHierarchyPropertyRequested = FPropertyBagInstanceDataDetails::FPropertyBagPropertyDelegate::CreateSP(this, &FPropertyBagDetails::NavigateToHierarchyProperty);
+	HierarchyParams.HierarchyViewModel = HierarchyViewModel.IsValid() ? HierarchyViewModel.Get() : nullptr;
 	FPropertyBagInstanceDataDetails::FConstructParams Params
 	{
 		.BagStructProperty = StructProperty,
@@ -2912,7 +3161,16 @@ FReply FPropertyBagDetails::SummonHierarchyEditor(TSharedPtr<IPropertyHandle> In
 	
 	Window->SetTitle(Title);
 	
-	TSharedPtr<SPropertyBagHierarchyEditor> NewPropertyBagHierarchyEditor = SNew(SPropertyBagHierarchyEditor, InPropertyBagHandle.ToSharedRef(), InPropertyUtilities.ToSharedRef())
+	TSharedPtr<FPropertyPath> PropertyPath = InPropertyBagHandle->CreateFPropertyPath();
+	if (!PropertyPath.IsValid())
+	{
+		return FReply::Unhandled();
+	}
+
+	FStructUtilsEditorModule& Module = FModuleManager::GetModuleChecked<FStructUtilsEditorModule>("StructUtilsEditor");
+	TNotNull<UPropertyBagHierarchyViewModel*> ViewModel = Module.FindOrCreateHierarchyViewModel(HierarchyRoot, OuterObjects[0], PropertyPath.ToSharedRef(), InPropertyUtilities.ToSharedRef());
+
+	TSharedPtr<SPropertyBagHierarchyEditor> NewPropertyBagHierarchyEditor = SNew(SPropertyBagHierarchyEditor, ViewModel)
 	.OnCloseRequested_Static(&FPropertyBagDetails::ClosePropertyBagHierarchyEditorWindow);
 	
 	PropertyBagHierarchyEditor = NewPropertyBagHierarchyEditor;
diff --git a/Engine/Source/Editor/StructUtilsEditor/Private/HierarchyEditor/PropertyBagHierarchyViewModel.cpp b/Engine/Source/Editor/StructUtilsEditor/Private/HierarchyEditor/PropertyBagHierarchyViewModel.cpp
--- a/Engine/Source/Editor/StructUtilsEditor/Private/HierarchyEditor/PropertyBagHierarchyViewModel.cpp
+++ b/Engine/Source/Editor/StructUtilsEditor/Private/HierarchyEditor/PropertyBagHierarchyViewModel.cpp
@@ -409,39 +409,46 @@ FString FPropertyBagHierarchyPropertyViewModel::ToString() const
 	return Desc ? Desc->Name.ToString() : TEXT("Unidentified");
 }
 
-void UPropertyBagHierarchyViewModel::Initialize(TSharedRef<IPropertyHandle> InOriginalPropertyBagHandle, TSharedRef<IPropertyUtilities> InPropertyUtilities)
+void UPropertyBagHierarchyViewModel::Initialize(UObject* InOwningObject, TSharedRef<FPropertyPath> InPropertyPath)
 {
+	if (!InOwningObject)
+	{
+		return;
+	}
+
+	OwningObject = InOwningObject;
+	PropertyPathToPropertyBag = InPropertyPath;
+
+	// We generate our own property handle via a PropertyRowGenerator. This gives us a stable handle
+	// that survives details panel refreshes (which would invalidate the original handle).
 	FPropertyRowGeneratorArgs Args;
 	PropertyRowGenerator = FModuleManager::Get().LoadModuleChecked<FPropertyEditorModule>("PropertyEditor").CreatePropertyRowGenerator(Args);
-	
-	PropertyUtilities = InPropertyUtilities;
-	PropertyBagSchema = UE::StructUtils::ExtractPropertyBagSchemaCDO(InOriginalPropertyBagHandle);
-	HierarchyRoot = UE::StructUtils::ExtractHierarchyRoot(InOriginalPropertyBagHandle);
-	
-	TArray<UObject*> OuterObjects;
-	InOriginalPropertyBagHandle->GetOuterObjects(OuterObjects);
-	
-	if (OuterObjects.Num() != 1)
+	PropertyRowGenerator->SetObjects({InOwningObject});
+	RegeneratePropertyBagHandle();
+
+	// Extract schema and hierarchy root from our generated handle
+	if (GeneratedPropertyBagHandle.IsValid())
 	{
-		return;
+		PropertyBagSchema = UE::StructUtils::ExtractPropertyBagSchemaCDO(GeneratedPropertyBagHandle.ToSharedRef());
+		HierarchyRoot = UE::StructUtils::ExtractHierarchyRoot(GeneratedPropertyBagHandle.ToSharedRef());
 	}
-	
-	OwningObject = OuterObjects[0];
-	
-	// The original handle will become invalidated if the details panel does a full refresh. We need to ensure it stays valid.
-	// We cache the property path which allows us to find our generated handle again
-	PropertyPathToPropertyBag = InOriginalPropertyBagHandle->CreateFPropertyPath();
-	
-	// We generate our own property handle now. Even this one can become invalid if changes to the property bag are made, but we can regenerate it now
-	PropertyRowGenerator->SetObjects(OuterObjects);
-	RegeneratePropertyBagHandle();	
-	
+
 	// We add the properties that don't currently exist in the hierarchy to the root.
 	AddMissingPropertiesToRoot(false);
-	
+
+	// Initialize the hash to avoid a false-positive change detection on the first tick
+	LastHash = UE::StructUtils::CalcPropertyDescArrayHash(GetPropertyDescs());
+
 	Super::Initialize();
 }
 
+void UPropertyBagHierarchyViewModel::BeginDestroy()
+{
+	Finalize();
+	
+	Super::BeginDestroy();
+}
+
 void UPropertyBagHierarchyViewModel::RegeneratePropertyBagHandle()
 {
 	if (!PropertyPathToPropertyBag.IsValid() || !PropertyRowGenerator.IsValid())
@@ -481,6 +488,11 @@ TConstArrayView<FPropertyBagPropertyDesc> UPropertyBagHierarchyViewModel::GetPro
 	return {};
 }
 
+void UPropertyBagHierarchyViewModel::SetPropertyUtilities(TSharedRef<IPropertyUtilities> InPropertyUtilities)
+{
+	PropertyUtilities = InPropertyUtilities;
+}
+
 void UPropertyBagHierarchyViewModel::RebuildDetails() const
 {
 	// We can't request a rebuild from the handle itself since it's not the same handle as in the details panel.
@@ -498,6 +510,15 @@ void UPropertyBagHierarchyViewModel::Tick(float DeltaTime)
 	{
 		RegeneratePropertyBagHandle();
 	}
+
+	// Detect structural changes to the property bag (new/removed properties) and refresh accordingly
+	uint64 CurrentHash = UE::StructUtils::CalcPropertyDescArrayHash(GetPropertyDescs());
+	if (LastHash != CurrentHash)
+	{
+		AddMissingPropertiesToRoot(false);
+		RequestFullRefreshNextFrame();
+		LastHash = CurrentHash;
+	}
 }
 
 bool UPropertyBagHierarchyViewModel::IsTickable() const
diff --git a/Engine/Source/Editor/StructUtilsEditor/Private/HierarchyEditor/SPropertyBagHierarchyEditor.cpp b/Engine/Source/Editor/StructUtilsEditor/Private/HierarchyEditor/SPropertyBagHierarchyEditor.cpp
--- a/Engine/Source/Editor/StructUtilsEditor/Private/HierarchyEditor/SPropertyBagHierarchyEditor.cpp
+++ b/Engine/Source/Editor/StructUtilsEditor/Private/HierarchyEditor/SPropertyBagHierarchyEditor.cpp
@@ -5,6 +5,7 @@
 #include "HierarchyEditor/PropertyBagHierarchyViewModel.h"
 #include "PropertyPath.h"
 #include "SlateOptMacros.h"
+#include "StructUtilsEditorLog.h"
 #include "StructUtilsEditorUtilsPrivate.h"
 #include "StructUtilsMetadata.h"
 #include "Framework/Notifications/NotificationManager.h"
@@ -17,28 +18,41 @@
 
 #define LOCTEXT_NAMESPACE "PropertyBagHierarchyEditor"
 
-void SPropertyBagHierarchyEditor::Construct(const FArguments& InArgs, TSharedRef<IPropertyHandle> InPropertyBagHandle, TSharedRef<IPropertyUtilities> InPropertyUtilities)
+void SPropertyBagHierarchyEditor::Construct(const FArguments& InArgs, TNotNull<UPropertyBagHierarchyViewModel*> InHierarchyViewModel)
 {
 	OnCloseRequested = InArgs._OnCloseRequested;
-	
-	PropertyBagSchemaCDO = UE::StructUtils::ExtractPropertyBagSchemaCDO(InPropertyBagHandle);
-	HierarchyRoot = UE::StructUtils::ExtractHierarchyRoot(InPropertyBagHandle);
-	
-	if (!PropertyBagSchemaCDO.IsValid() || !HierarchyRoot.IsValid())
+
+	HierarchyViewModel.Reset(InHierarchyViewModel);
+
+	TSharedPtr<IPropertyHandle> PropertyBagHandle = HierarchyViewModel->GetPropertyBagHandle();
+	if (!PropertyBagHandle.IsValid())
+	{
+		UE_LOG(LogStructUtilsEditor, Warning, TEXT("SPropertyBagHierarchyEditor::Construct — HierarchyViewModel has no valid property bag handle."));
+		return;
+	}
+
+	PropertyBagSchemaCDO = UE::StructUtils::ExtractPropertyBagSchemaCDO(PropertyBagHandle.ToSharedRef());
+	HierarchyRoot = HierarchyViewModel->GetHierarchyRoot();
+
+	if (!PropertyBagSchemaCDO.IsValid())
 	{
+		UE_LOG(LogStructUtilsEditor, Warning, TEXT("SPropertyBagHierarchyEditor::Construct — Failed to extract PropertyBagSchema CDO."));
+		return;
+	}
+
+	if (!HierarchyRoot.IsValid())
+	{
+		UE_LOG(LogStructUtilsEditor, Warning, TEXT("SPropertyBagHierarchyEditor::Construct — HierarchyViewModel has no valid hierarchy root."));
 		return;
 	}
 	
 	TArray<UObject*> OutOuterObjects;
-	InPropertyBagHandle->GetOuterObjects(OutOuterObjects);
+	PropertyBagHandle->GetOuterObjects(OutOuterObjects);
 	Algo::Transform(OutOuterObjects, OuterObjects, [](UObject* Object)
 	{
 		return TWeakObjectPtr<UObject>(Object);
 	});
-	
-	HierarchyViewModel.Reset(NewObject<UPropertyBagHierarchyViewModel>(GetTransientPackage()));
-	HierarchyViewModel->Initialize(InPropertyBagHandle, InPropertyUtilities);
-	
+
 	if (GEditor)
 	{
 		if (UAssetEditorSubsystem* AssetEditorSubsystem = GEditor->GetEditorSubsystem<UAssetEditorSubsystem>())
@@ -54,29 +68,13 @@ void SPropertyBagHierarchyEditor::Construct(const FArguments& InArgs, TSharedRef
 	]; 
 }
 
-void SPropertyBagHierarchyEditor::Tick(const FGeometry& AllottedGeometry, const double InCurrentTime, const float InDeltaTime)
-{
-	if (!HierarchyViewModel.IsValid())
-	{
-		return;
-	}
-	
-	uint64 CurrentHash = UE::StructUtils::CalcPropertyDescArrayHash(HierarchyViewModel->GetPropertyDescs());
-	if (LastPropertyBagHash != CurrentHash)
-	{
-		HierarchyViewModel->AddMissingPropertiesToRoot(false);
-		HierarchyViewModel->RequestFullRefreshNextFrame();
-		LastPropertyBagHash = CurrentHash;
-	}
-}
-
 TSharedPtr<IPropertyHandle> SPropertyBagHierarchyEditor::GetPropertyBagHandle() const
-{ 
+{
 	if (HierarchyViewModel.IsValid())
 	{
-		return HierarchyViewModel->GetPropertyBagHandle(); 
+		return HierarchyViewModel->GetPropertyBagHandle();
 	}
-	
+
 	return nullptr;
 }
 
@@ -89,13 +87,8 @@ SPropertyBagHierarchyEditor::~SPropertyBagHierarchyEditor()
 			AssetEditorSubsystem->OnAssetClosedInEditor().RemoveAll(this);
 		}
 	}
-	
-	if (HierarchyViewModel.IsValid())
-	{
-		HierarchyViewModel->Finalize();
-		HierarchyViewModel->ConditionalBeginDestroy();
-		HierarchyViewModel.Reset();
-	}
+
+	HierarchyViewModel.Reset();
 }
 
 void SPropertyBagHierarchyEditor::NavigateToProperty(const FPropertyBagPropertyDesc& PropertyDesc, bool bAddIfNotFound)
@@ -113,7 +106,7 @@ void SPropertyBagHierarchyEditor::NavigateToProperty(const FPropertyBagPropertyD
 		if (bAddIfNotFound)
 		{			
 			UHierarchyElement* Root = HierarchyViewModel->GetHierarchyRootViewModel()->GetDataMutable();
-			
+
 			UPropertyBagHierarchyProperty* NewHierarchyProperty = NewObject<UPropertyBagHierarchyProperty>(Root);
 			NewHierarchyProperty->Initialize(PropertyDesc);
 			Root->GetChildrenMutable().Add(NewHierarchyProperty);
@@ -123,8 +116,8 @@ void SPropertyBagHierarchyEditor::NavigateToProperty(const FPropertyBagPropertyD
 			HierarchyViewModel->SetActiveHierarchySection(nullptr);
 			
 			HierarchyViewModel->RebuildDetails();
-			
-			FNotificationInfo Info(FText::FormatOrdered(LOCTEXT("MissingHierarchyPropertyAdded", "Property {0} was added to the hierarchy so it can be edited."), 
+
+			FNotificationInfo Info(FText::FormatOrdered(LOCTEXT("MissingHierarchyPropertyAdded", "Property {0} was added to the hierarchy so it can be edited."),
 				FText::FromName(PropertyDesc.Name)));
 			
 			FSlateNotificationManager::Get().AddNotification(Info);
diff --git a/Engine/Source/Editor/StructUtilsEditor/Private/HierarchyEditor/SPropertyBagHierarchyEditor.h b/Engine/Source/Editor/StructUtilsEditor/Private/HierarchyEditor/SPropertyBagHierarchyEditor.h
--- a/Engine/Source/Editor/StructUtilsEditor/Private/HierarchyEditor/SPropertyBagHierarchyEditor.h
+++ b/Engine/Source/Editor/StructUtilsEditor/Private/HierarchyEditor/SPropertyBagHierarchyEditor.h
@@ -4,6 +4,7 @@
 
 #include "CoreMinimal.h"
 #include "DataHierarchyViewModelBase.h"
+#include "Misc/NotNull.h"
 #include "PropertyBagDetails.h"
 #include "EdGraph/EdGraphSchema.h"
 #include "HierarchyEditor/PropertyBagHierarchyViewModel.h"
@@ -11,11 +12,13 @@
 #include "Templates/SubclassOf.h"
 #include "Widgets/SCompoundWidget.h"
 
+#define UE_API STRUCTUTILSEDITOR_API
+
 class UPropertyBagSchema;
 class UPropertyBagHierarchyViewModel;
 
-/** A wrapper around the actual hierarchy editor. Owns the view model. */
-class STRUCTUTILSEDITOR_API SPropertyBagHierarchyEditor : public SCompoundWidget
+/** A wrapper around the actual hierarchy editor. Receives and holds a shared reference to the view model. */
+class SPropertyBagHierarchyEditor : public SCompoundWidget
 {
 public:
 	SLATE_BEGIN_ARGS(SPropertyBagHierarchyEditor)
@@ -23,33 +26,27 @@ public:
 		SLATE_EVENT(FSimpleDelegate, OnCloseRequested)
 	SLATE_END_ARGS()
 
-	/**
-	 * @param InArgs 
-	 * @param InPropertyBagHandle For initialization purposes only as it is expected to get invalidated. Internally we store another handle that we generate. 
-	 * @param InPropertyUtilities Required to refresh the details panel that originally spawned this hierarchy editor
-	 */
-	void Construct(const FArguments& InArgs, TSharedRef<IPropertyHandle> InPropertyBagHandle, TSharedRef<IPropertyUtilities> InPropertyUtilities);
-	virtual ~SPropertyBagHierarchyEditor() override;
-	virtual void Tick(const FGeometry& AllottedGeometry, const double InCurrentTime, const float InDeltaTime) override;
-	
+	UE_API void Construct(const FArguments& InArgs, TNotNull<UPropertyBagHierarchyViewModel*> InHierarchyViewModel);
+	UE_API virtual ~SPropertyBagHierarchyEditor() override;
+
 	/** Returns the property handle that we generated ourselves via property row generator. */
-	TSharedPtr<IPropertyHandle> GetPropertyBagHandle() const;
-	void NavigateToProperty(const FPropertyBagPropertyDesc& PropertyDesc, bool bAddIfNotFound = false);
+	UE_API TSharedPtr<IPropertyHandle> GetPropertyBagHandle() const;
+	UE_API void NavigateToProperty(const FPropertyBagPropertyDesc& PropertyDesc, bool bAddIfNotFound = false);
 
 private:
 	TSharedRef<SWidget> OnGenerateRowContentWidget(TSharedRef<FHierarchyElementViewModel> ElementViewModel) const;
 	void OnAssetEditorClosed(UObject* Object, IAssetEditorInstance* AssetEditorInstance) const;
-	
+
 private:
 	TWeakObjectPtr<const UPropertyBagSchema> PropertyBagSchemaCDO;
+	/** Shared ownership — keeps the view model alive as long as this widget exists, even if the details panel is destroyed. */
 	TStrongObjectPtr<UPropertyBagHierarchyViewModel> HierarchyViewModel;
-	
+
 	TWeakObjectPtr<UHierarchyRoot> HierarchyRoot;
-	
+
 	TArray<TWeakObjectPtr<UObject>> OuterObjects;
-	
-	/** We recalculate this on tick. When differences are detected, we refresh the hierarchy editor. */
-	uint64 LastPropertyBagHash = 0;
-	
+
 	FSimpleDelegate OnCloseRequested;
 };
+
+#undef UE_API
diff --git a/Engine/Source/Editor/StructUtilsEditor/Private/PropertyBagDragDropHandler.cpp b/Engine/Source/Editor/StructUtilsEditor/Private/PropertyBagDragDropHandler.cpp
--- a/Engine/Source/Editor/StructUtilsEditor/Private/PropertyBagDragDropHandler.cpp
+++ b/Engine/Source/Editor/StructUtilsEditor/Private/PropertyBagDragDropHandler.cpp
@@ -2,10 +2,15 @@
 
 #include "PropertyBagDragDropHandler.h"
 
+#include "DataHierarchyViewModelBase.h"
 #include "IPropertyBagEditorGraph.h"
+#include "PropertyBagDetails.h"
 #include "StructUtilsMetadata.h"
 #include "EdGraph/EdGraph.h"
 #include "Input/Reply.h"
+#include "Kismet2/BlueprintEditorUtils.h"
+#include "Widgets/Images/SImage.h"
+#include "Widgets/Text/STextBlock.h"
 
 #define LOCTEXT_NAMESPACE "StructUtilsEditor"
 
@@ -24,7 +29,26 @@ FPropertyBagDetailsDragDropOp::FDecoration::FDecoration(const FText& InMessage,
 FPropertyBagDetailsDragDropOp::FPropertyBagDetailsDragDropOp(const FPropertyBagPropertyDesc& InPropertyDesc)
 	: PropertyDesc(InPropertyDesc)
 {
+	FEdGraphPinType PinType = UE::StructUtils::GetPropertyDescAsPin(InPropertyDesc);
+
+	TSharedRef<SWidget> DefaultDecorator = SNew(SHorizontalBox)
+	+ SHorizontalBox::Slot()
+	.AutoWidth()
+	.Padding(2.f)
+	[
+		SNew(SImage)
+		.Image(FBlueprintEditorUtils::GetIconFromPin(PinType))
+		.ColorAndOpacity(GetDefault<UPropertyBagSchema>()->GetPinTypeColor(PinType))
+	]
+	+ SHorizontalBox::Slot()
+	.VAlign(VAlign_Center)
+	[
+		SNew(STextBlock).Text(FText::FromName(InPropertyDesc.Name))
+	];
+
 	FGraphEditorDragDropAction::Construct();
+
+	SetFeedbackMessage(DefaultDecorator);
 }
 
 void FPropertyBagDetailsDragDropOp::SetDecoration(const EPropertyBagDropState NewDropState, TOptional<FDecoration> OverriddenDecoration)
@@ -152,6 +176,16 @@ FReply FPropertyBagDetailsDragDropOp::DroppedOnPanel(const TSharedRef<SWidget>&
 	return FReply::Handled();
 }
 
+void FPropertyBagDetailsDragDropOp::SetHierarchyElement(TSharedPtr<FHierarchyElementViewModel> InElement)
+{
+	HierarchyElement = InElement;
+}
+
+TSharedPtr<FHierarchyElementViewModel> FPropertyBagDetailsDragDropOp::GetHierarchyElement() const
+{
+	return HierarchyElement.Pin();
+}
+
 IPropertyBagEdGraphDragAndDrop* FPropertyBagDetailsDragDropOp::GetPropertyBagEdGraphDragAndDropInterface() const
 {
 	UEdGraph* Graph = GetHoveredGraph();
@@ -169,12 +203,17 @@ TSharedPtr<FDragDropOperation> FPropertyBagDetailsDragDropHandlerTarget::CreateD
 
 void FPropertyBagDetailsDragDropHandlerTarget::BindCanAcceptDragDrop(FCanAcceptPropertyBagDetailsRowDropOp&& CanAcceptDragDrop)
 {
-	CanAcceptDetailsRowDropOp = std::move(CanAcceptDragDrop);
+	CanAcceptDetailsRowDropOp = MoveTemp(CanAcceptDragDrop);
 }
 
 void FPropertyBagDetailsDragDropHandlerTarget::BindOnHandleDragDrop(FOnPropertyBagDetailsRowDropOp&& OnDragDrop)
 {
-	OnHandleDetailsRowDropOp = std::move(OnDragDrop);
+	OnHandleDetailsRowDropOp = MoveTemp(OnDragDrop);
+}
+
+void FPropertyBagDetailsDragDropHandlerTarget::BindOnHandleHierarchyElementDrop(FOnPropertyBagHierarchyElementDropOp&& OnDrop)
+{
+	OnHandleHierarchyElementDropOp = MoveTemp(OnDrop);
 }
 
 TOptional<EItemDropZone> FPropertyBagDetailsDragDropHandlerTarget::CanAcceptDrop(const FDragDropEvent& DragDropSource, const EItemDropZone DropZone) const
@@ -193,9 +232,21 @@ TOptional<EItemDropZone> FPropertyBagDetailsDragDropHandlerTarget::CanAcceptDrop
 
 bool FPropertyBagDetailsDragDropHandlerTarget::AcceptDrop(const FDragDropEvent& DragDropSource, const EItemDropZone DropZone) const
 {
-	// Property Bag Details Drag and Drop Handler
 	if (const TSharedPtr<FPropertyBagDetailsDragDropOp> DropOp = DragDropSource.GetOperationAs<FPropertyBagDetailsDragDropOp>())
 	{
+		// Try hierarchy element path first — gets the element directly from the op without a PropertyDesc identity roundtrip
+		if (OnHandleHierarchyElementDropOp.IsBound())
+		{
+			if (TSharedPtr<FHierarchyElementViewModel> Element = DropOp->GetHierarchyElement())
+			{
+				return OnHandleHierarchyElementDropOp.Execute(Element, DropZone).IsEventHandled();
+			}
+			
+			// The view model should always be valid but in case it's not, something went wrong; don't accept the drop.
+			return false;
+		}
+
+		// Regular property desc path for non-hierarchy property bags
 		if (OnHandleDetailsRowDropOp.IsBound())
 		{
 			return OnHandleDetailsRowDropOp.Execute(DropOp->PropertyDesc, DropZone).IsEventHandled();
@@ -212,6 +263,7 @@ TSharedPtr<FDragDropOperation> FPropertyBagDetailsDragDropHandler::CreateDragDro
 {
 	if (TSharedPtr<FPropertyBagDetailsDragDropOp> DragOp = MakeShared<FPropertyBagDetailsDragDropOp>(PropertyDesc))
 	{
+		DragOp->SetHierarchyElement(HierarchyElement.Pin());
 		return DragOp;
 	}
 
diff --git a/Engine/Source/Editor/StructUtilsEditor/Private/StructUtilsEditorLog.cpp b/Engine/Source/Editor/StructUtilsEditor/Private/StructUtilsEditorLog.cpp
new file mode 100644
--- /dev/null
+++ b/Engine/Source/Editor/StructUtilsEditor/Private/StructUtilsEditorLog.cpp
@@ -0,0 +1,5 @@
+// Copyright Epic Games, Inc. All Rights Reserved.
+
+#include "StructUtilsEditorLog.h"
+
+DEFINE_LOG_CATEGORY(LogStructUtilsEditor);
diff --git a/Engine/Source/Editor/StructUtilsEditor/Private/StructUtilsEditorLog.h b/Engine/Source/Editor/StructUtilsEditor/Private/StructUtilsEditorLog.h
new file mode 100644
--- /dev/null
+++ b/Engine/Source/Editor/StructUtilsEditor/Private/StructUtilsEditorLog.h
@@ -0,0 +1,7 @@
+// Copyright Epic Games, Inc. All Rights Reserved.
+
+#pragma once
+
+#include "Logging/LogMacros.h"
+
+STRUCTUTILSEDITOR_API DECLARE_LOG_CATEGORY_EXTERN(LogStructUtilsEditor, Log, All);
diff --git a/Engine/Source/Editor/StructUtilsEditor/Private/StructUtilsEditorModule.cpp b/Engine/Source/Editor/StructUtilsEditor/Private/StructUtilsEditorModule.cpp
--- a/Engine/Source/Editor/StructUtilsEditor/Private/StructUtilsEditorModule.cpp
+++ b/Engine/Source/Editor/StructUtilsEditor/Private/StructUtilsEditorModule.cpp
@@ -3,6 +3,7 @@
 #include "StructUtilsEditorModule.h"
 #include "CoreGlobals.h"
 #include "GameFramework/Actor.h"
+#include "HierarchyEditor/PropertyBagHierarchyViewModel.h"
 #include "Modules/ModuleManager.h"
 #include "PropertyBagDetails.h"
 #include "PropertyEditorModule.h"
@@ -39,12 +40,17 @@ void FStructUtilsEditorModule::StartupModule()
 	PropertyModule.RegisterCustomPropertyTypeLayout(FInstancedPropertyBagTypeName, FOnGetPropertyTypeCustomizationInstance::CreateStatic(&FPropertyBagDetails::MakeInstance));
 	PropertyModule.RegisterCustomPropertyTypeLayout(FInstancedStructArrayTypeName, FOnGetPropertyTypeCustomizationInstance::CreateStatic(&FInstancedStructArrayDetails::MakeInstance));
 	PropertyModule.NotifyCustomizationModuleChanged();
+
+	PostGarbageCollectHandle = FCoreUObjectDelegates::GetPostGarbageCollect().AddRaw(this, &FStructUtilsEditorModule::OnPostGarbageCollect);
 }
 
 void FStructUtilsEditorModule::ShutdownModule()
 {
 	using namespace UE::StructUtilsEditor::Private;
 
+	FCoreUObjectDelegates::GetPostGarbageCollect().Remove(PostGarbageCollectHandle);
+	HierarchyViewModelMap.Empty();
+
 	// Unregister the details customization
 	if (FModuleManager::Get().IsModuleLoaded("PropertyEditor"))
 	{
@@ -56,6 +62,45 @@ void FStructUtilsEditorModule::ShutdownModule()
 	}
 }
 
+TNotNull<UPropertyBagHierarchyViewModel*> FStructUtilsEditorModule::FindOrCreateHierarchyViewModel(
+	UPropertyBagHierarchyRoot* InHierarchyRoot,
+	UObject* InOwningObject,
+	TSharedRef<FPropertyPath> InPropertyPath,
+	TSharedRef<IPropertyUtilities> PropertyUtilities)
+{
+	// Look up existing
+	if (TWeakObjectPtr<UPropertyBagHierarchyViewModel>* Existing = HierarchyViewModelMap.Find(InHierarchyRoot))
+	{
+		if (UPropertyBagHierarchyViewModel* ExistingVM = Existing->Get())
+		{
+			// Refresh property utilities so RebuildDetails works with the current details panel
+			ExistingVM->SetPropertyUtilities(PropertyUtilities);
+			return ExistingVM;
+		}
+	}
+
+	// Create new - add to map BEFORE Initialize to prevent reentrancy.
+	// Initialize() creates a PropertyRowGenerator which triggers a layout rebuild,
+	// which re-enters FindOrCreateHierarchyViewModel. The map entry ensures the recursive call
+	// finds the existing (not-yet-initialized) VM and returns it instead of creating another.
+	UPropertyBagHierarchyViewModel* NewVM = NewObject<UPropertyBagHierarchyViewModel>(GetTransientPackage());
+	HierarchyViewModelMap.Add(InHierarchyRoot, NewVM);
+	NewVM->Initialize(InOwningObject, InPropertyPath);
+	NewVM->SetPropertyUtilities(PropertyUtilities);
+	return NewVM;
+}
+
+void FStructUtilsEditorModule::OnPostGarbageCollect()
+{
+	for (auto It = HierarchyViewModelMap.CreateIterator(); It; ++It)
+	{
+		if (!It->Key.IsValid() || !It->Value.IsValid())
+		{
+			It.RemoveCurrent();
+		}
+	}
+}
+
 void FStructUtilsEditorModule::PreChange(const UUserDefinedStruct* StructToReinstantiate, FStructureEditorUtils::EStructureEditorChangeInfo Info)
 {
 }
diff --git a/Engine/Source/Editor/StructUtilsEditor/Public/HierarchyEditor/PropertyBagHierarchyViewModel.h b/Engine/Source/Editor/StructUtilsEditor/Public/HierarchyEditor/PropertyBagHierarchyViewModel.h
--- a/Engine/Source/Editor/StructUtilsEditor/Public/HierarchyEditor/PropertyBagHierarchyViewModel.h
+++ b/Engine/Source/Editor/StructUtilsEditor/Public/HierarchyEditor/PropertyBagHierarchyViewModel.h
@@ -188,17 +188,27 @@ protected:
 };
 
 /** The hierarchy view model for property bags. Defines core rules for how the hierarchy behaves. */
-UCLASS()
-class STRUCTUTILSEDITOR_API UPropertyBagHierarchyViewModel : public UDataHierarchyViewModelBase, public FTickableEditorObject
+UCLASS(MinimalAPI)
+class UPropertyBagHierarchyViewModel : public UDataHierarchyViewModelBase, public FTickableEditorObject
 {
 	GENERATED_BODY()
 	
 public:
 	/**
-	 * @param InOriginalPropertyBagHandle For initialization purposes only as it is expected to get invalidated. Internally we store another handle that we generate. 
-	 * @param InPropertyUtilities Required to refresh the details panel that originally spawned this hierarchy editor
+	 * Initializes the hierarchy view model. Internally creates a PropertyRowGenerator to generate a stable property handle
+	 * that survives details panel refreshes. The property handle is used by ApplyChangesToSinglePropertyDesc to notify
+	 * pre/post change when editing property metadata (tooltips, clamp values, etc.) via the hierarchy editor.
+	 *
+	 * Note: PropertyUtilities must be set separately via SetPropertyUtilities() after Initialize returns.
+	 * This is because Initialize triggers a PropertyRowGenerator layout rebuild which re-enters FindOrCreateHierarchyViewModel,
+	 * and setting utilities during that reentrancy would overwrite them with the generator's internal utilities.
+	 *
+	 * @param InOwningObject The UObject that owns the property bag (e.g. the Blueprint CDO or actor). Used to populate the PropertyRowGenerator.
+	 * @param InPropertyPath Path to the FInstancedPropertyBag property within the owning object. Used to locate the property bag handle in the generated property tree.
 	 */
-	void Initialize(TSharedRef<IPropertyHandle> InOriginalPropertyBagHandle, TSharedRef<IPropertyUtilities> InPropertyUtilities);
+	void Initialize(UObject* InOwningObject, TSharedRef<FPropertyPath> InPropertyPath);
+	
+	virtual void BeginDestroy() override;
 	
 	void RegeneratePropertyBagHandle();
 	
@@ -210,15 +220,19 @@ public:
 	TSharedPtr<IPropertyHandle> GetPropertyBagHandle() const { return GeneratedPropertyBagHandle; }
 	TSharedPtr<IPropertyUtilities> GetPropertyUtilities() const { return PropertyUtilities.IsValid() ? PropertyUtilities.Pin() : nullptr; }
 	TSharedRef<FPropertyPath> GetCurrentPropertyPath() const { return PropertyPathToPropertyBag.ToSharedRef(); }
-	
+
+	/** Updates the property utilities reference. Call when a new details panel instance provides fresh utilities. */
+	void SetPropertyUtilities(TSharedRef<IPropertyUtilities> InPropertyUtilities);
+
 	void RebuildDetails() const;
 
+	virtual UHierarchyRoot* GetHierarchyRoot() const override;
+
 protected:
 	virtual void Tick(float DeltaTime) override;
 	virtual bool IsTickable() const override;
 	virtual TStatId GetStatId() const override;
 	
-	virtual UHierarchyRoot* GetHierarchyRoot() const override;
 	virtual void OnHierarchyChanged(const TInstancedStruct<FHierarchyChangedPayload>& Payload) override;
 	virtual TSubclassOf<UHierarchyCategory> GetCategoryDataClass() const override;
 	virtual TSubclassOf<UHierarchySection> GetSectionDataClass() const override;
diff --git a/Engine/Source/Editor/StructUtilsEditor/Public/PropertyBagDetails.h b/Engine/Source/Editor/StructUtilsEditor/Public/PropertyBagDetails.h
--- a/Engine/Source/Editor/StructUtilsEditor/Public/PropertyBagDetails.h
+++ b/Engine/Source/Editor/StructUtilsEditor/Public/PropertyBagDetails.h
@@ -28,6 +28,7 @@ class UPropertyBagHierarchyRoot;
 class UPropertyBagHierarchyProperty;
 class UPropertyBagHierarchySection;
 class UPropertyBagHierarchyCategory;
+class UPropertyBagHierarchyViewModel;
 
 /**
  * The property bag details child rows can be completely customized by assigning a combination of these values
@@ -162,19 +163,23 @@ protected:
 	void NavigateToHierarchyProperty(const FPropertyBagPropertyDesc& Desc);
 	
 	static void ClosePropertyBagHierarchyEditorWindow();
+
 private:
 	/** Handle to the property bag being edited */
 	TSharedPtr<IPropertyHandle> StructProperty;
 	TSharedPtr<IPropertyUtilities> PropUtils;
 	EPropertyBagPropertyType DefaultType = EPropertyBagPropertyType::Bool;
 	EPropertyBagChildRowFeatures ChildRowFeatures = EPropertyBagChildRowFeatures::Default;
-	
+
 	// We keep track of existing window & hierarchy editor loosely so we can recycle them where appropriate
 	// This also means we can only have 1 property bag hierarchy editor open at a time
 	// This is because an additional window can't rely on the existing details panel and property handles too much
 	static TWeakPtr<SWindow> PropertyBagHierarchyEditorWindow;
 	static TWeakPtr<class SPropertyBagHierarchyEditor> PropertyBagHierarchyEditor;
 
+	/** Strong reference to the hierarchy view model for this customization instance's lifetime. */
+	TStrongObjectPtr<UPropertyBagHierarchyViewModel> HierarchyViewModel;
+
 	UE_DEPRECATED(5.6, "Use 'ChildRowFeatures' instead.")
 	bool bFixedLayout = false;
 	bool bAllowContainers = true;
@@ -204,6 +209,7 @@ public:
 		struct FHierarchyParams
 		{
 			FPropertyBagPropertyDelegate OnNavigateToHierarchyPropertyRequested;
+			TWeakObjectPtr<UPropertyBagHierarchyViewModel> HierarchyViewModel;
 		};
 		
 		TSharedPtr<IPropertyHandle> BagStructProperty = nullptr;
@@ -304,10 +310,13 @@ public:
 	/**~ Public Hierarchy Editor API */	
 	UE_API EVisibility ShouldShowHierarchyProperty(TSharedPtr<IPropertyHandle> ChildPropertyHandle) const;
 	
-	TWeakObjectPtr<const UHierarchySection> GetActiveHierarchySection() const { return HierarchyData.ActiveHierarchySection; }
-	
+	UE_API TWeakObjectPtr<const UHierarchySection> GetActiveHierarchySection() const;
+
 	/** Retrieves the hierarchy property handle map. Only meaningful if used in conjunction with a hierarchy root. */
-	const TMap<TObjectKey<const UPropertyBagHierarchyProperty>, TSharedPtr<IPropertyHandle>>& GetHierarchyPropertyHandleMap() const { return HierarchyData.HierarchyPropertyHandleMap; }
+	UE_API const TMap<TObjectKey<const UPropertyBagHierarchyProperty>, TSharedPtr<IPropertyHandle>>& GetHierarchyPropertyHandleMap() const;
+
+	/** Retrieves the hierarchy view model, if available. */
+	UE_API UPropertyBagHierarchyViewModel* GetHierarchyViewModel() const;
 	/**~ End of Hierarchy Editor API */
 	
 protected:
@@ -348,6 +357,7 @@ protected:
 	struct FHierarchyData
 	{
 		TWeakObjectPtr<UPropertyBagHierarchyRoot> HierarchyRoot;
+		TWeakObjectPtr<UPropertyBagHierarchyViewModel> HierarchyViewModel;
 		TMap<TObjectKey<const UPropertyBagHierarchyProperty>, TSharedPtr<IPropertyHandle>> HierarchyPropertyHandleMap;
 		TMap<FString, TWeakObjectPtr<const UPropertyBagHierarchyProperty>> PropertyPathPropertyMap;
 		TWeakObjectPtr<const UHierarchySection> ActiveHierarchySection;
diff --git a/Engine/Source/Editor/StructUtilsEditor/Public/PropertyBagDragDropHandler.h b/Engine/Source/Editor/StructUtilsEditor/Public/PropertyBagDragDropHandler.h
--- a/Engine/Source/Editor/StructUtilsEditor/Public/PropertyBagDragDropHandler.h
+++ b/Engine/Source/Editor/StructUtilsEditor/Public/PropertyBagDragDropHandler.h
@@ -10,9 +10,11 @@
 
 class FPropertyBagDetailsDragDropOp;
 class IPropertyBagEdGraphDragAndDrop;
+struct FHierarchyElementViewModel;
 
 DECLARE_DELEGATE_RetVal_TwoParams(TOptional<EItemDropZone>, FCanAcceptPropertyBagDetailsRowDropOp, TSharedPtr<FPropertyBagDetailsDragDropOp> /*DropOperation*/, EItemDropZone /*DropZone*/);
 DECLARE_DELEGATE_RetVal_TwoParams(FReply, FOnPropertyBagDetailsRowDropOp, const FPropertyBagPropertyDesc& /*PropertyDescription*/, EItemDropZone /*DropZone*/);
+DECLARE_DELEGATE_RetVal_TwoParams(FReply, FOnPropertyBagHierarchyElementDropOp, TSharedPtr<FHierarchyElementViewModel> /*DroppedElement*/, EItemDropZone /*DropZone*/);
 
 /** State of the drop, useful for source-is-target validation on details rows. */
 enum class EPropertyBagDropState
@@ -23,7 +25,7 @@ enum class EPropertyBagDropState
 };
 
 /**
- * Provides information about the source row (single property) being dragged.
+ * Provides information about the source row (single property or hierarchy element) being dragged.
  * Inherits from FGraphEditorDragDropAction to support dragging to a graph pin/node/panel.
  */
 class FPropertyBagDetailsDragDropOp : public FGraphEditorDragDropAction
@@ -41,7 +43,7 @@ public:
 	};
 
 	/**
-	 * Property Bag Details - Drop Operation Constructor
+	 * Property Bag Details - Drop Operation Constructor for property elements.
 	 * @param InPropertyDesc Property description for the dropped property.
 	 */
 	UE_API FPropertyBagDetailsDragDropOp(const FPropertyBagPropertyDesc& InPropertyDesc);
@@ -62,73 +64,43 @@ public:
 	/** Visibility of the cursor decorator error icon. */
 	UE_API virtual EVisibility GetErrorIconVisible() const override;
 
-	/**
-	 * The drop operation has been executed on an UEdGraphPin in an Editor Graph.
-	 * @param ScreenPosition 2D screen coordinates.
-	 * @param GraphPosition Local 2D graph coordinates.
-	 * @return If the drop operation was handled by the graph.
-	 */
 	UE_API virtual FReply DroppedOnPin(const FVector2f& ScreenPosition, const FVector2f& GraphPosition) override;
-
-	/**
-	 * The drop operation has been executed on an UEdGraphNode in an Editor Graph.
-	 * @param ScreenPosition 2D screen coordinates.
-	 * @param GraphPosition Local 2D graph coordinates.
-	 * @return If the drop operation was handled by the graph.
-	 */
 	UE_API virtual FReply DroppedOnNode(const FVector2f& ScreenPosition, const FVector2f& GraphPosition) override;
-
-	/**
-	 * The drop operation has been executed on an SGraphPanel in an Editor Graph.
-	 * @param Panel The graph panel widget to receive the drop action.
-	 * @param ScreenPosition 2D screen coordinates.
-	 * @param GraphPosition Local 2D graph coordinates.
-	 * @param Graph The graph owner of the target panel.
-	 * @return If the drop operation was handled by the graph.
-	 */
 	UE_API virtual FReply DroppedOnPanel(const TSharedRef<SWidget>& Panel, const FVector2f& ScreenPosition, const FVector2f& GraphPosition, UEdGraph& Graph) override;
 	// ~End FGraphEditorDragDropAction interface
 
-	/** The property description of the dragged property. */
+	/** The property description of the dragged property. May be default-constructed for non-property hierarchy elements. */
 	const FPropertyBagPropertyDesc PropertyDesc;
 	/** Cached state for the decorator. */
 	EPropertyBagDropState CurrentDropState = EPropertyBagDropState::Invalid;
 
+	/** Optional hierarchy element view model for hierarchy-based drag & drop in the details panel. */
+	UE_API void SetHierarchyElement(TSharedPtr<FHierarchyElementViewModel> InElement);
+	UE_API TSharedPtr<FHierarchyElementViewModel> GetHierarchyElement() const;
+
 private:
 	UE_API IPropertyBagEdGraphDragAndDrop* GetPropertyBagEdGraphDragAndDropInterface() const;
+
+	TWeakPtr<FHierarchyElementViewModel> HierarchyElement;
 };
 
 /** Handles drag-and-drop (as a target) for a single property's child widget row. */
 class FPropertyBagDetailsDragDropHandlerTarget : public IDetailDragDropHandler
 {
 public:
-	/**
-	 * Property Bag Details - Drag/Drop Handler Default Constructor
-	 */
 	FPropertyBagDetailsDragDropHandlerTarget() = default;
 
-	/**
-	 * Property Bag Details - Drag/Drop Handler Constructor
-	 * @param CanAcceptDragDrop [Signature]: TOptional<EItemDropZone>(TSharedPtr<FPropertyBagDetailsDragDropOp>, EItemDropZone)
-	 * @param OnDragDrop [Signature]: FReply(const FPropertyBagPropertyDesc&, EItemDropZone)
-	 */
 	UE_API FPropertyBagDetailsDragDropHandlerTarget(const FCanAcceptPropertyBagDetailsRowDropOp& CanAcceptDragDrop, const FOnPropertyBagDetailsRowDropOp& OnDragDrop);
 
 	/** Creates the drag and drop operation. Disabled for the target. Enabled in subclasses. */
 	UE_API virtual TSharedPtr<FDragDropOperation> CreateDragDropOperation() const override;
 
-	/**
-	 * Bind the delegate for determining if this handler can accept this drag and drop operation.
-	 * @param CanAcceptDragDrop [Signature]: TOptional<EItemDropZone>(TSharedPtr<FPropertyBagDetailsDragDropOp>, EItemDropZone)
-	 */
 	UE_API virtual void BindCanAcceptDragDrop(FCanAcceptPropertyBagDetailsRowDropOp&& CanAcceptDragDrop);
-
-	/**
-	 * Bind the delegate to handle a valid drag and drop operation.
-	 * @param OnDragDrop [Signature]: FReply(const FPropertyBagPropertyDesc&, EItemDropZone)
-	 */
 	UE_API virtual void BindOnHandleDragDrop(FOnPropertyBagDetailsRowDropOp&& OnDragDrop);
 
+	/** Bind a hierarchy-element-based drop handler. Used when the dropped element may not have a valid PropertyDesc (e.g. categories). */
+	UE_API virtual void BindOnHandleHierarchyElementDrop(FOnPropertyBagHierarchyElementDropOp&& OnDrop);
+
 	/** Disable automatic creation of the handle widget for targets. It is overridden in the source below. */
 	virtual bool UseHandleWidget() const override
 	{
@@ -136,51 +108,25 @@ public:
 	}
 
 protected:
-	/**
-	 * Whether the details view child row can accept a drag and drop operation. 
-	 * @param DragDropSource The source dragged onto this row.
-	 * @param DropZone Zone of the drop, relative to this child row.
-	 * @return [Optional] Target zone, which can be overridden. False if the row will not accept.
-	 */
 	UE_API virtual TOptional<EItemDropZone> CanAcceptDrop(const FDragDropEvent& DragDropSource, EItemDropZone DropZone) const override;
-
-	/**
-	 * Accept and handle the drag and drop operation.
-	 * @param DragDropSource The source dragged onto this row.
-	 * @param DropZone Zone of the drop, relative to this child row.
-	 * @return True if the operation was successful.
-	 */
 	UE_API virtual bool AcceptDrop(const FDragDropEvent& DragDropSource, EItemDropZone DropZone) const override;
 
 private:
-	/**
-	 * Callback to determine if the target is acceptable for a details view child row drag and drop operation.
-	 * Signature: Optional<EItemDropZone>(TSharedPtr<FPropertyBagDetailsDragDropOp>, EItemDropZone)
-	 */
 	FCanAcceptPropertyBagDetailsRowDropOp CanAcceptDetailsRowDropOp;
-
-	/**
-	 * Callback to determine if the drag to a details view child row was handled successfully.
-	 * Signature: FReply(const FPropertyBagPropertyDesc&, EItemDropZone)
-	 */
 	FOnPropertyBagDetailsRowDropOp OnHandleDetailsRowDropOp;
+	FOnPropertyBagHierarchyElementDropOp OnHandleHierarchyElementDropOp;
 };
 
 /** Handles drag-and-drop (as a source or target) for a single property's child widget row. */
 class FPropertyBagDetailsDragDropHandler : public FPropertyBagDetailsDragDropHandlerTarget
 {
 public:
-	/**
-	 * Constructor.
-	 * @param InPropertyDesc The property desc of the property represented by this child widget row.
-	 */
 	UE_API FPropertyBagDetailsDragDropHandler(const FPropertyBagPropertyDesc& InPropertyDesc);
 
+	/** Set an optional hierarchy element view model to attach to drag ops created by this handler. */
+	void SetHierarchyElement(TSharedPtr<FHierarchyElementViewModel> InElement) { HierarchyElement = InElement; }
+
 protected:
-	/**
-	 * Creates the drag and drop operation object that contains necessary information for a drag and drop operation.
-	 * @return Shared pointer to the drag and drop operation.
-	 */
 	UE_API virtual TSharedPtr<FDragDropOperation> CreateDragDropOperation() const override;
 
 	/** Enable the automatic creation of the handle (grip) widget for the source. */
@@ -189,6 +135,9 @@ protected:
 private:
 	/** The current child row's property bag property description. */
 	const FPropertyBagPropertyDesc PropertyDesc;
+
+	/** Optional hierarchy element for hierarchy-based drag & drop. */
+	TWeakPtr<FHierarchyElementViewModel> HierarchyElement;
 };
 
 #undef UE_API
diff --git a/Engine/Source/Editor/StructUtilsEditor/Public/StructUtilsEditorModule.h b/Engine/Source/Editor/StructUtilsEditor/Public/StructUtilsEditorModule.h
--- a/Engine/Source/Editor/StructUtilsEditor/Public/StructUtilsEditorModule.h
+++ b/Engine/Source/Editor/StructUtilsEditor/Public/StructUtilsEditorModule.h
@@ -2,6 +2,7 @@
 
 #pragma  once
 
+#include "Misc/NotNull.h"
 #include "Modules/ModuleInterface.h"
 #include "Kismet2/StructureEditorUtils.h"
 
@@ -10,6 +11,10 @@
 class IStructUtilsEditor;
 struct FGraphPanelNodeFactory;
 class UUserDefinedStruct;
+class UPropertyBagHierarchyRoot;
+class UPropertyBagHierarchyViewModel;
+class IPropertyUtilities;
+class FPropertyPath;
 
 /**
 * The public interface to this module
@@ -21,12 +26,27 @@ public:
 	UE_API virtual void StartupModule() override;
 	UE_API virtual void ShutdownModule() override;
 
+	/** Finds an existing or creates a new hierarchy view model for a given hierarchy root. */
+	UE_API TNotNull<UPropertyBagHierarchyViewModel*> FindOrCreateHierarchyViewModel(
+		UPropertyBagHierarchyRoot* InHierarchyRoot,
+		UObject* InOwningObject,
+		TSharedRef<FPropertyPath> InPropertyPath,
+		TSharedRef<IPropertyUtilities> PropertyUtilities);
+
 protected:
 
 	// INotifyOnStructChanged
 	UE_API virtual void PreChange(const class UUserDefinedStruct* Struct, FStructureEditorUtils::EStructureEditorChangeInfo Info) override;
 	UE_API virtual void PostChange(const class UUserDefinedStruct* Struct, FStructureEditorUtils::EStructureEditorChangeInfo Info) override;
 	// ~INotifyOnStructChanged
+
+private:
+	void OnPostGarbageCollect();
+
+	/** View model lookup by hierarchy root. Both keys and values are weak — consumers hold strong refs to the values. */
+	TMap<TWeakObjectPtr<UPropertyBagHierarchyRoot>, TWeakObjectPtr<UPropertyBagHierarchyViewModel>> HierarchyViewModelMap;
+
+	FDelegateHandle PostGarbageCollectHandle;
 };
 
 #undef UE_API
