Request a details panel refresh instead of forcing it when picking an instanced struct
Forcing a refresh invalidates the property handles that might have been cached by the OnPicked callback

#Synced-CL 52066653

diff --git a/Engine/Source/Editor/StructUtilsEditor/Private/SInstancedStructPicker.cpp b/Engine/Source/Editor/StructUtilsEditor/Private/SInstancedStructPicker.cpp
--- a/Engine/Source/Editor/StructUtilsEditor/Private/SInstancedStructPicker.cpp
+++ b/Engine/Source/Editor/StructUtilsEditor/Private/SInstancedStructPicker.cpp
@@ -442,7 +442,7 @@ void SInstancedStructPicker::StructPicked(const UScriptStruct* InStruct)
 		// Property tree will be invalid after changing the struct type, force update.
 		if (PropUtils.IsValid())
 		{
-			PropUtils->ForceRefresh();
+			PropUtils->RequestRefresh();
 		}
 	}
 
