Instanced Struct Picker Improvements:
- Add support for fully qualified GetAllowedClasses / GetDisallowedClasses static UFUNCTIONs (i.e. /Script/ModuleName.ClassName:FunctionName).
- Add ExplicitAllowedStructs meta flag with which it's possible to disallow child structs of allowed ones.

#rnx

#Synced-CL 52298949

diff --git a/Engine/Source/Editor/StructUtilsEditor/Private/SInstancedStructPicker.cpp b/Engine/Source/Editor/StructUtilsEditor/Private/SInstancedStructPicker.cpp
--- a/Engine/Source/Editor/StructUtilsEditor/Private/SInstancedStructPicker.cpp
+++ b/Engine/Source/Editor/StructUtilsEditor/Private/SInstancedStructPicker.cpp
@@ -69,7 +69,7 @@ bool FInstancedStructFilter::IsStructAllowed(const FStructViewerInitializationOp
 		bStructAllowed = false;
 		for (const TSoftObjectPtr<const UScriptStruct>& AllowedStruct : AllowedStructs)
 		{
-			if (InStruct->IsChildOf(AllowedStruct.Get()))
+			if (bAllowChildStructs ? InStruct->IsChildOf(AllowedStruct.Get()) : InStruct == AllowedStruct)
 			{
 				bStructAllowed = true;
 				break;
@@ -80,7 +80,7 @@ bool FInstancedStructFilter::IsStructAllowed(const FStructViewerInitializationOp
 	{
 		for (const TSoftObjectPtr<const UScriptStruct>& DisallowedStruct : DisallowedStructs)
 		{
-			if (InStruct->IsChildOf(DisallowedStruct.Get()))
+			if (bAllowChildStructs ? InStruct->IsChildOf(DisallowedStruct.Get()) : DisallowedStruct == InStruct)
 			{
 				bStructAllowed = false;
 				break;
@@ -282,10 +282,12 @@ const FSlateBrush* SInstancedStructPicker::GetDisplayValueIcon() const
 TSharedRef<SWidget> SInstancedStructPicker::GenerateStructPicker()
 {
 	static const FName NAME_ExcludeBaseStruct = "ExcludeBaseStruct";
+	static const FName NAME_ExcludeChildStructs = "ExcludeChildStructs";
 	static const FName NAME_HideViewOptions = "HideViewOptions";
 	static const FName NAME_ShowTreeView = "ShowTreeView";
 
 	const bool bExcludeBaseStruct = StructProperty->HasMetaData(NAME_ExcludeBaseStruct);
+	const bool bExcludeChildStructs = StructProperty->HasMetaData(NAME_ExcludeChildStructs);
 	const bool bAllowNone = !(StructProperty->GetMetaDataProperty()->PropertyFlags & CPF_NoClear);
 	const bool bHideViewOptions = StructProperty->HasMetaData(NAME_HideViewOptions);
 	const bool bShowTreeView = StructProperty->HasMetaData(NAME_ShowTreeView);
@@ -296,6 +298,7 @@ TSharedRef<SWidget> SInstancedStructPicker::GenerateStructPicker()
 	StructFilter->BaseStruct = BaseScriptStruct;
 	StructFilter->bAllowUserDefinedStructs = BaseScriptStruct == nullptr; // Only allow user defined structs when BaseStruct is not set.
 	StructFilter->bAllowBaseStruct = !bExcludeBaseStruct;
+	StructFilter->bAllowChildStructs = !bExcludeChildStructs;
 
 	if (GEditor && StructProperty)
 	{
@@ -323,33 +326,40 @@ TSharedRef<SWidget> SInstancedStructPicker::GenerateStructPicker()
 
 		TArray<UObject*> OwningObjects;
 		StructProperty->GetOuterObjects(OwningObjects);
-		for (UObject* OwningObject : OwningObjects)
+
+		DECLARE_DELEGATE_RetVal(TArray<TSoftObjectPtr<UScriptStruct>>, FGetClasses);
+
+		auto GatherStructs = [this, &OwningObjects](const FName MetaName, TArray<TSoftObjectPtr<const UScriptStruct>>& OutStructs)
 		{
-			if (OwningObject != nullptr)
+			const FString& FunctionName = StructProperty->GetMetaData(MetaName);
+			if (!FunctionName.IsEmpty())
 			{
-				const FString GetAllowedClassesFunctionName = StructProperty->GetMetaData("GetAllowedClasses");
-				if (!GetAllowedClassesFunctionName.IsEmpty())
+				// First look for a library function assuming fully qualified path, e.g. /Script/ModuleName.ClassName:FunctionName
+				if (FunctionName.Contains(TEXT(".")))
 				{
-					const UFunction* GetAllowedClassesFunction = OwningObject ? OwningObject->FindFunction(*GetAllowedClassesFunctionName) : nullptr;
-					if (GetAllowedClassesFunction != nullptr)
+					const UFunction* Function = FindObject<UFunction>(nullptr, *FunctionName, EFindObjectFlags::ExactClass);
+					if (ensureMsgf(Function && Function->HasAnyFunctionFlags(EFunctionFlags::FUNC_Static), TEXT("Invalid %s: %s"), *MetaName.ToString(), *FunctionName))
 					{
-						DECLARE_DELEGATE_RetVal(TArray<TSoftObjectPtr<UScriptStruct>>, FGetAllowedClasses);
-						StructFilter->AllowedStructs.Append(FGetAllowedClasses::CreateUFunction(OwningObject, GetAllowedClassesFunction->GetFName()).Execute());
+						UObject* FunctionCDO = Function->GetOwnerClass()->GetDefaultObject();
+						OutStructs.Append(FGetClasses::CreateUFunction(FunctionCDO, Function->GetFName()).Execute());
 					}
 				}
-
-				const FString GetDisallowedClassesFunctionName = StructProperty->GetMetaData("GetDisallowedClasses");
-				if (!GetDisallowedClassesFunctionName.IsEmpty())
+				else
 				{
-					const UFunction* GetDisallowedClassesFunction = OwningObject ? OwningObject->FindFunction(*GetDisallowedClassesFunctionName) : nullptr;
-					if (GetDisallowedClassesFunction != nullptr)
+					for (UObject* OwningObject : OwningObjects)
 					{
-						DECLARE_DELEGATE_RetVal(TArray<TSoftObjectPtr<UScriptStruct>>, FGetDisallowedClasses);
-						StructFilter->DisallowedStructs.Append(FGetDisallowedClasses::CreateUFunction(OwningObject, GetDisallowedClassesFunction->GetFName()).Execute());
+						const UFunction* Function = OwningObject ? OwningObject->FindFunction(*FunctionName) : nullptr;
+						if (Function != nullptr)
+						{
+							OutStructs.Append(FGetClasses::CreateUFunction(OwningObject, Function->GetFName()).Execute());
+						}
 					}
 				}
 			}
-		}
+		};
+
+		GatherStructs("GetAllowedClasses", StructFilter->AllowedStructs);
+		GatherStructs("GetDisallowedClasses", StructFilter->DisallowedStructs);
 	}
 	else
 	{
diff --git a/Engine/Source/Editor/StructUtilsEditor/Public/SInstancedStructPicker.h b/Engine/Source/Editor/StructUtilsEditor/Public/SInstancedStructPicker.h
--- a/Engine/Source/Editor/StructUtilsEditor/Public/SInstancedStructPicker.h
+++ b/Engine/Source/Editor/StructUtilsEditor/Public/SInstancedStructPicker.h
@@ -34,6 +34,9 @@ public:
 	// A flag controlling whether we allow to select the BaseStruct
 	bool bAllowBaseStruct = true;
 
+	// A flag controlling whether we allow child structs of AllowedStructs / DisallowedStructs
+	bool bAllowChildStructs = true;
+
 	UE_API virtual bool IsStructAllowed(const FStructViewerInitializationOptions& InInitOptions, const UScriptStruct* InStruct, TSharedRef<FStructViewerFilterFuncs> InFilterFuncs) override;
 	UE_API virtual bool IsUnloadedStructAllowed(const FStructViewerInitializationOptions& InInitOptions, const FSoftObjectPath& InStructPath, TSharedRef<FStructViewerFilterFuncs> InFilterFuncs) override;
 
